Trivial change to LensData to add image width and height #rb Simon.Therriault

diff --git a/Engine/Plugins/VirtualProduction/CameraCalibrationCore/Source/CameraCalibrationCore/Public/LensData.h b/Engine/Plugins/VirtualProduction/CameraCalibrationCore/Source/CameraCalibrationCore/Public/LensData.h
--- a/Engine/Plugins/VirtualProduction/CameraCalibrationCore/Source/CameraCalibrationCore/Public/LensData.h
+++ b/Engine/Plugins/VirtualProduction/CameraCalibrationCore/Source/CameraCalibrationCore/Public/LensData.h
@@ -47,6 +47,11 @@ public:
 	/** Width and height of the calibrated camera's sensor, in millimeters */
 	UPROPERTY(EditAnywhere, BlueprintReadWrite, Category = "Lens Info")
 	FVector2D SensorDimensions = FVector2D(23.76f, 13.365f);
+
+	/** Camera image size, in pixels */
+	UPROPERTY(EditAnywhere, BlueprintReadWrite, Category = "Lens Info")
+	FIntPoint ImageDimensions;
+
 };
 
 /**
