Fix deprecation warnings

diff --git a/Engine/Plugins/VirtualProduction/CameraCalibrationCore/Source/CameraCalibrationCore/Private/CalibratedMapProcessor.cpp b/Engine/Plugins/VirtualProduction/CameraCalibrationCore/Source/CameraCalibrationCore/Private/CalibratedMapProcessor.cpp
--- a/Engine/Plugins/VirtualProduction/CameraCalibrationCore/Source/CameraCalibrationCore/Private/CalibratedMapProcessor.cpp
+++ b/Engine/Plugins/VirtualProduction/CameraCalibrationCore/Source/CameraCalibrationCore/Private/CalibratedMapProcessor.cpp
@@ -148,9 +148,9 @@ void FCalibratedMapProcessor::ExecuteJob(TSharedPtr<FDerivedDistortionDataJob> J
 		if(SourceTexture == nullptr
 		|| DestinationUndistortionTexture == nullptr
 		|| DestinationDistortionTexture == nullptr
-		|| SourceTexture->Resource == nullptr
-		|| DestinationUndistortionTexture->Resource == nullptr
-		|| DestinationDistortionTexture->Resource == nullptr)
+		|| SourceTexture->GetResource() == nullptr
+		|| DestinationUndistortionTexture->GetResource() == nullptr
+		|| DestinationDistortionTexture->GetResource() == nullptr)
 		{
 			Job->State = EDerivedDistortionDataJobState::Completed;
 			Job->Output.Result = EDerivedDistortionDataResult::Error;
@@ -158,9 +158,9 @@ void FCalibratedMapProcessor::ExecuteJob(TSharedPtr<FDerivedDistortionDataJob> J
 			return;
 		}
 		
-		FTextureResource* SourceDistortionMap = SourceTexture->Resource;
-		FTextureResource* DestinationUndistortionDisplacementMap = DestinationUndistortionTexture->Resource;
-		FTextureResource* DestinationDistortionDisplacementMap = DestinationDistortionTexture->Resource;
+		FTextureResource* SourceDistortionMap = SourceTexture->GetResource();
+		FTextureResource* DestinationUndistortionDisplacementMap = DestinationUndistortionTexture->GetResource();
+		FTextureResource* DestinationDistortionDisplacementMap = DestinationDistortionTexture->GetResource();
 
 		{
 			FScopeLock Lock(&RunningJobsCriticalSection);
