Fix to save the LensFile image dimensions in the asset

diff --git a/Engine/Plugins/VirtualProduction/CameraCalibrationCore/Source/CameraCalibrationCore/Public/LensData.h b/Engine/Plugins/VirtualProduction/CameraCalibrationCore/Source/CameraCalibrationCore/Public/LensData.h
--- a/Engine/Plugins/VirtualProduction/CameraCalibrationCore/Source/CameraCalibrationCore/Public/LensData.h
+++ b/Engine/Plugins/VirtualProduction/CameraCalibrationCore/Source/CameraCalibrationCore/Public/LensData.h
@@ -27,6 +27,7 @@ public:
 		return LensModelName == Other.LensModelName
 			&& LensSerialNumber == Other.LensSerialNumber
 			&& LensModel == Other.LensModel
+			&& ImageDimensions == Other.ImageDimensions
 			&& SensorDimensions == Other.SensorDimensions;
 	}
 	
