CameraCalibration: Fix circular dependency with Composure and LiveLinkCamera.

#rnx

diff --git a/Engine/Plugins/Compositing/Composure/Composure.uplugin b/Engine/Plugins/Compositing/Composure/Composure.uplugin
--- a/Engine/Plugins/Compositing/Composure/Composure.uplugin
+++ b/Engine/Plugins/Compositing/Composure/Composure.uplugin
@@ -34,7 +34,7 @@
 	],
 	"Plugins": [
 		{
-			"Name": "CameraCalibration",
+			"Name": "CameraCalibrationCore",
 			"Enabled": true
 		},
 		{
diff --git a/Engine/Plugins/Compositing/Composure/Source/Composure/Composure.Build.cs b/Engine/Plugins/Compositing/Composure/Source/Composure/Composure.Build.cs
--- a/Engine/Plugins/Compositing/Composure/Source/Composure/Composure.Build.cs
+++ b/Engine/Plugins/Compositing/Composure/Source/Composure/Composure.Build.cs
@@ -16,7 +16,7 @@ namespace UnrealBuildTool.Rules
             PublicDependencyModuleNames.AddRange(
 				new string[]
 				{
-					"CameraCalibration",
+					"CameraCalibrationCore",
 					"CinematicCamera",
 					"Core",
                     "CoreUObject",
diff --git a/Engine/Plugins/VirtualProduction/CameraCalibration/CameraCalibration.uplugin b/Engine/Plugins/VirtualProduction/CameraCalibration/CameraCalibration.uplugin
--- a/Engine/Plugins/VirtualProduction/CameraCalibration/CameraCalibration.uplugin
+++ b/Engine/Plugins/VirtualProduction/CameraCalibration/CameraCalibration.uplugin
@@ -17,7 +17,7 @@
 		{
 			"Name": "CameraCalibration",
 			"Type": "Runtime",
-			"LoadingPhase": "PostConfigInit"
+			"LoadingPhase": "Default"
 		},
 		{
 			"Name": "CameraCalibrationEditor",
@@ -26,6 +26,10 @@
 		}
 	],
 	"Plugins": [
+		{
+			"Name": "CameraCalibrationCore",
+			"Enabled": true
+		},
 		{
 			"Name": "LiveLink",
 			"Enabled": true
diff --git a/Engine/Plugins/VirtualProduction/CameraCalibration/Config/BaseCameraCalibration.ini b/Engine/Plugins/VirtualProduction/CameraCalibration/Config/BaseCameraCalibration.ini
deleted file mode 100644
diff --git a/Engine/Plugins/VirtualProduction/CameraCalibration/Shaders/Private/DisplacementMapBlending.usf b/Engine/Plugins/VirtualProduction/CameraCalibration/Shaders/Private/DisplacementMapBlending.usf
deleted file mode 100644
diff --git a/Engine/Plugins/VirtualProduction/CameraCalibration/Shaders/Private/DistortionSTMapProcessor.usf b/Engine/Plugins/VirtualProduction/CameraCalibration/Shaders/Private/DistortionSTMapProcessor.usf
deleted file mode 100644
diff --git a/Engine/Plugins/VirtualProduction/CameraCalibration/Source/CameraCalibration/CameraCalibration.Build.cs b/Engine/Plugins/VirtualProduction/CameraCalibration/Source/CameraCalibration/CameraCalibration.Build.cs
--- a/Engine/Plugins/VirtualProduction/CameraCalibration/Source/CameraCalibration/CameraCalibration.Build.cs
+++ b/Engine/Plugins/VirtualProduction/CameraCalibration/Source/CameraCalibration/CameraCalibration.Build.cs
@@ -17,15 +17,11 @@ namespace UnrealBuildTool.Rules
 			PublicDependencyModuleNames.AddRange(
 				new string[]
 				{
+					"CameraCalibrationCore",
 					"CinematicCamera",
 					"Core",
 					"CoreUObject",
-					"DeveloperSettings",
 					"Engine",
-					"RenderCore",
-					"RHI",
-					"Slate",
-					"SlateCore",
 				}
 			);
 
@@ -33,22 +29,8 @@ namespace UnrealBuildTool.Rules
 				new string[]
                 {
 					"LiveLinkComponents",
-                    "Projects",
 				}
 			);
-
-			if (Target.bBuildEditor == true)
-			{
-				PrivateDependencyModuleNames.Add("UnrealEd");
-			}
-
-			PrivateIncludePaths.AddRange(
-				new string[] 
-				{
-            		//required for ScreenPass
-            		"../../../../Source/Runtime/Renderer/Private",
-            	}
-			);
 		}
 	}
 }
diff --git a/Engine/Plugins/VirtualProduction/CameraCalibration/Source/CameraCalibration/Private/CalibratedMapProcessor.cpp b/Engine/Plugins/VirtualProduction/CameraCalibration/Source/CameraCalibration/Private/CalibratedMapProcessor.cpp
deleted file mode 100644
diff --git a/Engine/Plugins/VirtualProduction/CameraCalibration/Source/CameraCalibration/Private/CalibratedMapProcessor.h b/Engine/Plugins/VirtualProduction/CameraCalibration/Source/CameraCalibration/Private/CalibratedMapProcessor.h
deleted file mode 100644
diff --git a/Engine/Plugins/VirtualProduction/CameraCalibration/Source/CameraCalibration/Private/CameraCalibrationModule.cpp b/Engine/Plugins/VirtualProduction/CameraCalibration/Source/CameraCalibration/Private/CameraCalibrationModule.cpp
--- a/Engine/Plugins/VirtualProduction/CameraCalibration/Source/CameraCalibration/Private/CameraCalibrationModule.cpp
+++ b/Engine/Plugins/VirtualProduction/CameraCalibration/Source/CameraCalibration/Private/CameraCalibrationModule.cpp
@@ -3,152 +3,16 @@
 #include "CameraCalibrationModule.h"
 
 #include "CameraCalibrationLog.h"
-#include "CameraCalibrationSettings.h"
-#include "CameraCalibrationSubsystem.h"
-#include "Engine/Engine.h"
-#include "HAL/IConsoleManager.h"
-#include "Interfaces/IPluginManager.h"
-#include "LensFile.h"
-#include "Misc/App.h"
-#include "Misc/CoreDelegates.h"
-#include "Misc/Paths.h"
-#include "Models/SphericalLensModel.h"
 #include "Modules/ModuleManager.h"
-#include "ShaderCore.h"
-#include "UObject/UObjectGlobals.h"
-
 
 DEFINE_LOG_CATEGORY(LogCameraCalibration);
 
-
-static TAutoConsoleVariable<FString> CVarStartupLensFile(
-	TEXT("CameraCalibration.StartupLensFile"),
-	TEXT(""),
-	TEXT("Startup Lens File\n"),
-	ECVF_ReadOnly
-);
-
-
-
-
 void FCameraCalibrationModule::StartupModule()
 {
-	FString PluginShaderDir = FPaths::Combine(IPluginManager::Get().FindPlugin(TEXT("CameraCalibration"))->GetBaseDir(), TEXT("Shaders"));
-	AddShaderSourceDirectoryMapping(TEXT("/Plugin/CameraCalibration"), PluginShaderDir);
-
-	ApplyStartupLensFile();
-	RegisterDistortionModels();
 }
 
-
 void FCameraCalibrationModule::ShutdownModule()
 {
-	UnregisterDistortionModels();
-
-	if (PostEngineInitHandle.IsValid())
-	{
-		FCoreDelegates::OnPostEngineInit.Remove(PostEngineInitHandle);
-	}
-}
-
-void FCameraCalibrationModule::ApplyStartupLensFile()
-{
-	auto ApplyLensFile = [this]()
-	{
-		ULensFile* StartupLensFile = nullptr;
-
-		// Try to load from CVar
-		{
-			const FString LensFileName = CVarStartupLensFile.GetValueOnGameThread();
-
-			if (LensFileName.Len())
-			{
-				if (UObject* Object = StaticLoadObject(ULensFile::StaticClass(), nullptr, *LensFileName))
-				{
-					StartupLensFile = CastChecked<ULensFile>(Object);
-				}
-
-				if (StartupLensFile)
-				{
-					UE_LOG(LogCameraCalibration, Display, TEXT("Loading Lens File specified in CVar CameraCalibration.StartupLensFile: '%s'"), *LensFileName);
-				}
-			}
-		}
-
-#if WITH_EDITOR
-		// Try to load user settings
-		if (StartupLensFile == nullptr)
-		{
-			StartupLensFile = GetDefault<UCameraCalibrationEditorSettings>()->GetUserLensFile();
-			if (StartupLensFile)
-			{
-				UE_LOG(LogCameraCalibration, Display, TEXT("Loading Lens File specified in user settings: '%s'"), *StartupLensFile->GetName());
-			}
-		}
-#endif
-
-		// Try loading default lens file from project settings
-		if (StartupLensFile == nullptr)
-		{
-			StartupLensFile = GetDefault<UCameraCalibrationSettings>()->GetStartupLensFile();
-
-			if (StartupLensFile)
-			{
-				UE_LOG(LogCameraCalibration, Display, TEXT("Loading Lens File specified in project settings: '%s'"), *StartupLensFile->GetName());
-			}
-		}
-
-		UCameraCalibrationSubsystem* SubSystem = GEngine->GetEngineSubsystem<UCameraCalibrationSubsystem>();
-		check(SubSystem);
-		SubSystem->SetDefaultLensFile(StartupLensFile);
-	};
-
-	if (FApp::CanEverRender())
-	{
-		if (GEngine && GEngine->IsInitialized())
-		{
-			ApplyLensFile();
-		}
-		else
-		{
-			PostEngineInitHandle = FCoreDelegates::OnPostEngineInit.AddLambda(ApplyLensFile);
-		}
-	}
-}
-
-void FCameraCalibrationModule::RegisterDistortionModels()
-{
-	auto RegisterModels = [this]()
-	{
-		// Register all lens models defined in this module
-		UCameraCalibrationSubsystem* SubSystem = GEngine->GetEngineSubsystem<UCameraCalibrationSubsystem>();
-		SubSystem->RegisterDistortionModel(USphericalLensModel::StaticClass());
-	};
-
-	if (FApp::CanEverRender())
-	{
-		if (GEngine && GEngine->IsInitialized())
-		{
-			RegisterModels();
-		}
-		else
-		{
-			PostEngineInitHandle = FCoreDelegates::OnPostEngineInit.AddLambda(RegisterModels);
-		}
-	}
 }
 
-void FCameraCalibrationModule::UnregisterDistortionModels()
-{
-	if (GEngine)
-	{
-		// Unregister all lens models defined in this module
-		if (UCameraCalibrationSubsystem* SubSystem = GEngine->GetEngineSubsystem<UCameraCalibrationSubsystem>())
-		{
-			SubSystem->UnregisterDistortionModel(USphericalLensModel::StaticClass());
-		}
-	}
-}
-
-
-IMPLEMENT_MODULE(FCameraCalibrationModule, CameraCalibration)
\ No newline at end of file
+IMPLEMENT_MODULE(FCameraCalibrationModule, CameraCalibration)
diff --git a/Engine/Plugins/VirtualProduction/CameraCalibration/Source/CameraCalibration/Private/CameraCalibrationModule.h b/Engine/Plugins/VirtualProduction/CameraCalibration/Source/CameraCalibration/Private/CameraCalibrationModule.h
--- a/Engine/Plugins/VirtualProduction/CameraCalibration/Source/CameraCalibration/Private/CameraCalibrationModule.h
+++ b/Engine/Plugins/VirtualProduction/CameraCalibration/Source/CameraCalibration/Private/CameraCalibrationModule.h
@@ -2,8 +2,8 @@
 
 #pragma once
 
+#include "CoreMinimal.h"
 #include "Modules/ModuleInterface.h"
-#include "Delegates/IDelegateInstance.h"
 
 
 class FCameraCalibrationModule : public IModuleInterface
@@ -14,13 +14,4 @@ public:
 	virtual void StartupModule() override;
 	virtual void ShutdownModule() override;
 	//~ End IModuleInterface interface
-
-private:
-	void ApplyStartupLensFile();
-	void RegisterDistortionModels();
-	void UnregisterDistortionModels();
-
-private:
-
-	FDelegateHandle PostEngineInitHandle;
 };
diff --git a/Engine/Plugins/VirtualProduction/CameraCalibration/Source/CameraCalibration/Private/CameraCalibrationSettings.cpp b/Engine/Plugins/VirtualProduction/CameraCalibration/Source/CameraCalibration/Private/CameraCalibrationSettings.cpp
deleted file mode 100644
diff --git a/Engine/Plugins/VirtualProduction/CameraCalibration/Source/CameraCalibration/Private/CameraCalibrationSubsystem.cpp b/Engine/Plugins/VirtualProduction/CameraCalibration/Source/CameraCalibration/Private/CameraCalibrationSubsystem.cpp
deleted file mode 100644
diff --git a/Engine/Plugins/VirtualProduction/CameraCalibration/Source/CameraCalibration/Private/CameraCalibrationUtils.cpp b/Engine/Plugins/VirtualProduction/CameraCalibration/Source/CameraCalibration/Private/CameraCalibrationUtils.cpp
deleted file mode 100644
diff --git a/Engine/Plugins/VirtualProduction/CameraCalibration/Source/CameraCalibration/Private/LensDistortionModelHandlerBase.cpp b/Engine/Plugins/VirtualProduction/CameraCalibration/Source/CameraCalibration/Private/LensDistortionModelHandlerBase.cpp
deleted file mode 100644
diff --git a/Engine/Plugins/VirtualProduction/CameraCalibration/Source/CameraCalibration/Private/LensFile.cpp b/Engine/Plugins/VirtualProduction/CameraCalibration/Source/CameraCalibration/Private/LensFile.cpp
deleted file mode 100644
diff --git a/Engine/Plugins/VirtualProduction/CameraCalibration/Source/CameraCalibration/Private/LensFileRendering.cpp b/Engine/Plugins/VirtualProduction/CameraCalibration/Source/CameraCalibration/Private/LensFileRendering.cpp
deleted file mode 100644
diff --git a/Engine/Plugins/VirtualProduction/CameraCalibration/Source/CameraCalibration/Private/LensFileRendering.h b/Engine/Plugins/VirtualProduction/CameraCalibration/Source/CameraCalibration/Private/LensFileRendering.h
deleted file mode 100644
diff --git a/Engine/Plugins/VirtualProduction/CameraCalibration/Source/CameraCalibration/Private/LensInterpolationUtils.cpp b/Engine/Plugins/VirtualProduction/CameraCalibration/Source/CameraCalibration/Private/LensInterpolationUtils.cpp
deleted file mode 100644
diff --git a/Engine/Plugins/VirtualProduction/CameraCalibration/Source/CameraCalibration/Private/LensInterpolationUtils.h b/Engine/Plugins/VirtualProduction/CameraCalibration/Source/CameraCalibration/Private/LensInterpolationUtils.h
deleted file mode 100644
diff --git a/Engine/Plugins/VirtualProduction/CameraCalibration/Source/CameraCalibration/Private/Models/LensModel.cpp b/Engine/Plugins/VirtualProduction/CameraCalibration/Source/CameraCalibration/Private/Models/LensModel.cpp
deleted file mode 100644
diff --git a/Engine/Plugins/VirtualProduction/CameraCalibration/Source/CameraCalibration/Private/Models/SphericalLensModel.cpp b/Engine/Plugins/VirtualProduction/CameraCalibration/Source/CameraCalibration/Private/Models/SphericalLensModel.cpp
deleted file mode 100644
diff --git a/Engine/Plugins/VirtualProduction/CameraCalibration/Source/CameraCalibration/Private/SphericalLensDistortionModelHandler.cpp b/Engine/Plugins/VirtualProduction/CameraCalibration/Source/CameraCalibration/Private/SphericalLensDistortionModelHandler.cpp
deleted file mode 100644
diff --git a/Engine/Plugins/VirtualProduction/CameraCalibration/Source/CameraCalibration/Private/Tests/TestDistortionMapping.cpp b/Engine/Plugins/VirtualProduction/CameraCalibration/Source/CameraCalibration/Private/Tests/TestDistortionMapping.cpp
deleted file mode 100644
diff --git a/Engine/Plugins/VirtualProduction/CameraCalibration/Source/CameraCalibration/Public/CalibrationPointComponent.h b/Engine/Plugins/VirtualProduction/CameraCalibration/Source/CameraCalibration/Public/CalibrationPointComponent.h
deleted file mode 100644
diff --git a/Engine/Plugins/VirtualProduction/CameraCalibration/Source/CameraCalibration/Public/CameraCalibrationSettings.h b/Engine/Plugins/VirtualProduction/CameraCalibration/Source/CameraCalibration/Public/CameraCalibrationSettings.h
deleted file mode 100644
diff --git a/Engine/Plugins/VirtualProduction/CameraCalibration/Source/CameraCalibration/Public/CameraCalibrationStep.h b/Engine/Plugins/VirtualProduction/CameraCalibration/Source/CameraCalibration/Public/CameraCalibrationStep.h
deleted file mode 100644
diff --git a/Engine/Plugins/VirtualProduction/CameraCalibration/Source/CameraCalibration/Public/CameraCalibrationSubsystem.h b/Engine/Plugins/VirtualProduction/CameraCalibration/Source/CameraCalibration/Public/CameraCalibrationSubsystem.h
deleted file mode 100644
diff --git a/Engine/Plugins/VirtualProduction/CameraCalibration/Source/CameraCalibration/Public/CameraCalibrationTypes.h b/Engine/Plugins/VirtualProduction/CameraCalibration/Source/CameraCalibration/Public/CameraCalibrationTypes.h
deleted file mode 100644
diff --git a/Engine/Plugins/VirtualProduction/CameraCalibration/Source/CameraCalibration/Public/CameraCalibrationUtils.h b/Engine/Plugins/VirtualProduction/CameraCalibration/Source/CameraCalibration/Public/CameraCalibrationUtils.h
deleted file mode 100644
diff --git a/Engine/Plugins/VirtualProduction/CameraCalibration/Source/CameraCalibration/Public/CameraNodalOffsetAlgo.h b/Engine/Plugins/VirtualProduction/CameraCalibration/Source/CameraCalibration/Public/CameraNodalOffsetAlgo.h
deleted file mode 100644
diff --git a/Engine/Plugins/VirtualProduction/CameraCalibration/Source/CameraCalibration/Public/ICalibratedMapProcessor.h b/Engine/Plugins/VirtualProduction/CameraCalibration/Source/CameraCalibration/Public/ICalibratedMapProcessor.h
deleted file mode 100644
diff --git a/Engine/Plugins/VirtualProduction/CameraCalibration/Source/CameraCalibration/Public/LensData.h b/Engine/Plugins/VirtualProduction/CameraCalibration/Source/CameraCalibration/Public/LensData.h
deleted file mode 100644
diff --git a/Engine/Plugins/VirtualProduction/CameraCalibration/Source/CameraCalibration/Public/LensDistortionModelHandlerBase.h b/Engine/Plugins/VirtualProduction/CameraCalibration/Source/CameraCalibration/Public/LensDistortionModelHandlerBase.h
deleted file mode 100644
diff --git a/Engine/Plugins/VirtualProduction/CameraCalibration/Source/CameraCalibration/Public/LensFile.h b/Engine/Plugins/VirtualProduction/CameraCalibration/Source/CameraCalibration/Public/LensFile.h
deleted file mode 100644
diff --git a/Engine/Plugins/VirtualProduction/CameraCalibration/Source/CameraCalibration/Public/Models/LensModel.h b/Engine/Plugins/VirtualProduction/CameraCalibration/Source/CameraCalibration/Public/Models/LensModel.h
deleted file mode 100644
diff --git a/Engine/Plugins/VirtualProduction/CameraCalibration/Source/CameraCalibration/Public/Models/SphericalLensModel.h b/Engine/Plugins/VirtualProduction/CameraCalibration/Source/CameraCalibration/Public/Models/SphericalLensModel.h
deleted file mode 100644
diff --git a/Engine/Plugins/VirtualProduction/CameraCalibration/Source/CameraCalibration/Public/SphericalLensDistortionModelHandler.h b/Engine/Plugins/VirtualProduction/CameraCalibration/Source/CameraCalibration/Public/SphericalLensDistortionModelHandler.h
deleted file mode 100644
diff --git a/Engine/Plugins/VirtualProduction/CameraCalibration/Source/CameraCalibrationEditor/CameraCalibrationEditor.Build.cs b/Engine/Plugins/VirtualProduction/CameraCalibration/Source/CameraCalibrationEditor/CameraCalibrationEditor.Build.cs
--- a/Engine/Plugins/VirtualProduction/CameraCalibration/Source/CameraCalibrationEditor/CameraCalibrationEditor.Build.cs
+++ b/Engine/Plugins/VirtualProduction/CameraCalibration/Source/CameraCalibrationEditor/CameraCalibrationEditor.Build.cs
@@ -6,20 +6,13 @@ namespace UnrealBuildTool.Rules
 	{
 		public CameraCalibrationEditor(ReadOnlyTargetRules Target) : base(Target)
 		{
-			if (Target.Configuration == UnrealTargetConfiguration.Debug)
-			{
-				PCHUsage = ModuleRules.PCHUsageMode.UseExplicitOrSharedPCHs;
-				OptimizeCode = CodeOptimization.Never;
-				bUseUnity = false;
-				PCHUsage = PCHUsageMode.NoPCHs;
-			}
-
 			PrivateDependencyModuleNames.AddRange(
 				new string[]
 				{
 					"AssetRegistry",
 					"AssetTools",
 					"CameraCalibration",
+					"CameraCalibrationCore",
 					"CinematicCamera",
 					"Composure",
 					"ComposureLayersEditor",
diff --git a/Engine/Plugins/VirtualProduction/CameraCalibration/Source/CameraCalibrationEditor/Private/CameraCalibrationEditorModule.cpp b/Engine/Plugins/VirtualProduction/CameraCalibration/Source/CameraCalibrationEditor/Private/CameraCalibrationEditorModule.cpp
--- a/Engine/Plugins/VirtualProduction/CameraCalibration/Source/CameraCalibrationEditor/Private/CameraCalibrationEditorModule.cpp
+++ b/Engine/Plugins/VirtualProduction/CameraCalibration/Source/CameraCalibrationEditor/Private/CameraCalibrationEditorModule.cpp
@@ -7,7 +7,6 @@
 #include "AssetToolsModule.h"
 #include "AssetTypeActions/AssetTypeActions_LensFile.h"
 #include "CameraCalibrationTypes.h"
-#include "DistortionHandlerPickerDetailCustomization.h"
 #include "IAssetTools.h"
 #include "IAssetTypeActions.h"
 #include "LensFile.h"
@@ -38,9 +37,6 @@ void FCameraCalibrationEditorModule::StartupModule()
 	// register asset type actions
 	RegisterAssetTypeAction(MakeShared<FAssetTypeActions_LensFile>());
 
-	// register detail panel customization
-	FPropertyEditorModule& PropertyModule = FModuleManager::GetModuleChecked<FPropertyEditorModule>("PropertyEditor");
-
 	{
 		const IWorkspaceMenuStructure& MenuStructure = WorkspaceMenu::GetMenuStructure();
 		TSharedRef<FWorkspaceItem> BrowserGroup = MenuStructure.GetDeveloperToolsMiscCategory()->GetParent()->AddGroup(
@@ -49,9 +45,6 @@ void FCameraCalibrationEditorModule::StartupModule()
 			true);
 	}
 
-	FPropertyEditorModule& PropertyEditorModule = FModuleManager::Get().LoadModuleChecked<FPropertyEditorModule>("PropertyEditor");
-	PropertyEditorModule.RegisterCustomPropertyTypeLayout(FDistortionHandlerPicker::StaticStruct()->GetFName(), FOnGetPropertyTypeCustomizationInstance::CreateStatic(&FDistortionHandlerPickerDetailCustomization::MakeInstance));
-
 	FCameraCalibrationMenuEntry::Register();
 }
 
@@ -63,10 +56,6 @@ void FCameraCalibrationEditorModule::ShutdownModule()
 
 		FLevelEditorModule* LevelEditorModule = FModuleManager::GetModulePtr<FLevelEditorModule>("LevelEditor");
 
-		FPropertyEditorModule& PropertyModule = FModuleManager::GetModuleChecked<FPropertyEditorModule>("PropertyEditor");
-		PropertyModule.UnregisterCustomClassLayout(ULensFile::StaticClass()->GetFName());
-		PropertyModule.UnregisterCustomPropertyTypeLayout(FDistortionHandlerPicker::StaticStruct()->GetFName());
-
 		// Unregister AssetTypeActions
 		FAssetToolsModule* AssetToolsModule = FModuleManager::GetModulePtr<FAssetToolsModule>("AssetTools");
 
diff --git a/Engine/Plugins/VirtualProduction/CameraCalibration/Source/CameraCalibrationEditor/Private/DistortionHandlerPickerDetailCustomization.cpp b/Engine/Plugins/VirtualProduction/CameraCalibration/Source/CameraCalibrationEditor/Private/DistortionHandlerPickerDetailCustomization.cpp
deleted file mode 100644
diff --git a/Engine/Plugins/VirtualProduction/CameraCalibration/Source/CameraCalibrationEditor/Private/DistortionHandlerPickerDetailCustomization.h b/Engine/Plugins/VirtualProduction/CameraCalibration/Source/CameraCalibrationEditor/Private/DistortionHandlerPickerDetailCustomization.h
deleted file mode 100644
diff --git a/Engine/Plugins/VirtualProduction/CameraCalibrationCore/CameraCalibrationCore.uplugin b/Engine/Plugins/VirtualProduction/CameraCalibrationCore/CameraCalibrationCore.uplugin
new file mode 100644
--- /dev/null
+++ b/Engine/Plugins/VirtualProduction/CameraCalibrationCore/CameraCalibrationCore.uplugin
@@ -0,0 +1,29 @@
+{
+	"FileVersion": 3,
+	"Version": 1,
+	"VersionName": "1.0",
+	"FriendlyName": "Camera Calibration Core",
+	"Description": "Supports lens distortion and camera calibration.",
+	"Category": "Virtual Production",
+	"CreatedBy": "Epic Games, Inc.",
+	"CreatedByURL": "http://epicgames.com",
+	"DocsURL": "",
+	"MarketplaceURL": "",
+	"SupportURL": "",
+	"CanContainContent": true,
+	"IsBetaVersion": true,
+	"Hidden": true,
+	"Installed": false,
+	"Modules": [
+		{
+			"Name": "CameraCalibrationCore",
+			"Type": "Runtime",
+			"LoadingPhase": "PostConfigInit"
+		},
+		{
+			"Name": "CameraCalibrationCoreEditor",
+			"Type": "Editor",
+			"LoadingPhase": "Default"
+		}
+	]
+}
\ No newline at end of file
diff --git a/Engine/Plugins/VirtualProduction/CameraCalibrationCore/Config/BaseCameraCalibrationCore.ini b/Engine/Plugins/VirtualProduction/CameraCalibrationCore/Config/BaseCameraCalibrationCore.ini
new file mode 100644
--- /dev/null
+++ b/Engine/Plugins/VirtualProduction/CameraCalibrationCore/Config/BaseCameraCalibrationCore.ini
@@ -0,0 +1,27 @@
+[CoreRedirects]
+
+;
+; 4.27 Redirects
+;
+
+; Classes
+
++ClassRedirects=(OldName="/Script/LensDistortion.LensFile", NewName="/Script/CameraCalibrationCore.LensFile")
++ClassRedirects=(OldName="/Script/LensDistortion.LensDistortionModelHandlerBase", NewName="/Script/CameraCalibrationCore.LensDistortionModelHandlerBase")
++ClassRedirects=(OldName="/Script/LensDistortion.SphericalLensDistortionModelHandler", NewName="/Script/CameraCalibrationCore.SphericalLensDistortionModelHandler")
++ClassRedirects=(OldName="/Script/LensDistortion.LensModel", NewName="/Script/CameraCalibrationCore.LensModel")
++ClassRedirects=(OldName="/Script/LensDistortion.SphericalLensModel", NewName="/Script/CameraCalibrationCore.SphericalLensModel")
+
++ClassRedirects=(OldName="/Script/CameraCalibration.LensFile", NewName="/Script/CameraCalibrationCore.LensFile")
++ClassRedirects=(OldName="/Script/CameraCalibration.LensDistortionModelHandlerBase", NewName="/Script/CameraCalibrationCore.LensDistortionModelHandlerBase")
++ClassRedirects=(OldName="/Script/CameraCalibration.SphericalLensDistortionModelHandler", NewName="/Script/CameraCalibrationCore.SphericalLensDistortionModelHandler")
++ClassRedirects=(OldName="/Script/CameraCalibration.LensModel", NewName="/Script/CameraCalibrationCore.LensModel")
++ClassRedirects=(OldName="/Script/CameraCalibration.SphericalLensModel", NewName="/Script/CameraCalibrationCore.SphericalLensModel")
+
+; Packages
+
++PackageRedirects=(OldName="/LensDistortion/Materials/M_SphericalDistortionPostProcess", NewName="/CameraCalibrationCore/Materials/M_DistortionPostProcess")
++PackageRedirects=(OldName="/LensDistortion/Materials/M_SphericalDistortionDisplacementMap", NewName="/CameraCalibrationCore/Materials/M_SphericalDistortionDisplacementMap")
+
++PackageRedirects=(OldName="/CameraCalibration/Materials/M_SphericalDistortionPostProcess", NewName="/CameraCalibrationCore/Materials/M_DistortionPostProcess")
++PackageRedirects=(OldName="/CameraCalibration/Materials/M_SphericalDistortionDisplacementMap", NewName="/CameraCalibrationCore/Materials/M_SphericalDistortionDisplacementMap")
diff --git a/Engine/Plugins/VirtualProduction/CameraCalibrationCore/Shaders/Private/DisplacementMapBlending.usf b/Engine/Plugins/VirtualProduction/CameraCalibrationCore/Shaders/Private/DisplacementMapBlending.usf
new file mode 100644
--- /dev/null
+++ b/Engine/Plugins/VirtualProduction/CameraCalibrationCore/Shaders/Private/DisplacementMapBlending.usf
@@ -0,0 +1,54 @@
+// Copyright Epic Games, Inc. All Rights Reserved.
+
+/*=============================================================================
+	DistortionSTMapProcessor.usf
+=============================================================================*/
+
+#include "/Engine/Private/Common.ush"
+
+float BlendFactor;
+float MainCoefficient;
+float DeltaMinX;
+float DeltaMaxX;
+float DeltaMinY;
+float DeltaMaxY;
+float2 FxFyScale;
+float2 PrincipalPoint;
+uint2 OutputTextureExtent;
+Texture2D SourceTextureOne;
+Texture2D SourceTextureTwo;
+Texture2D SourceTextureThree;
+Texture2D SourceTextureFour;
+SamplerState SourceTextureSampler;
+
+
+void BlendPS(float4 SvPosition : SV_POSITION, out float4 OutColor : SV_Target0)
+{
+	float2 UV = SvPosition.xy / OutputTextureExtent.xy;
+	UV = ((UV - 0.5f) * FxFyScale) + 0.5f;
+
+	float4 TexOneColor = Texture2DSampleLevel(SourceTextureOne, SourceTextureSampler, UV, 0);
+	float4 BlendedColor;
+
+#if BLEND_TYPE == 1
+	
+	float4 TexTwoColor = Texture2DSampleLevel(SourceTextureTwo, SourceTextureSampler, UV, 0);
+	BlendedColor = lerp(TexOneColor, TexTwoColor, BlendFactor);
+	
+#elif BLEND_TYPE == 2
+
+	float4 TexTwoColor = Texture2DSampleLevel(SourceTextureTwo, SourceTextureSampler, UV, 0);
+	float4 TexThreeColor = Texture2DSampleLevel(SourceTextureThree, SourceTextureSampler, UV, 0);
+	float4 TexFourColor = Texture2DSampleLevel(SourceTextureFour, SourceTextureSampler, UV, 0);
+
+	BlendedColor = MainCoefficient * (TexOneColor * DeltaMaxX * DeltaMaxY + TexTwoColor * DeltaMaxX * DeltaMinY + TexThreeColor * DeltaMinX * DeltaMaxY + TexFourColor * DeltaMinX * DeltaMinY);
+
+#else
+	
+	BlendedColor = TexOneColor;
+
+#endif
+	
+	float2 CenterShift = PrincipalPoint - float2(0.5f, 0.5f);
+	OutColor = BlendedColor + float4(CenterShift.xy, CenterShift.xy);
+}
diff --git a/Engine/Plugins/VirtualProduction/CameraCalibrationCore/Shaders/Private/DistortionSTMapProcessor.usf b/Engine/Plugins/VirtualProduction/CameraCalibrationCore/Shaders/Private/DistortionSTMapProcessor.usf
new file mode 100644
--- /dev/null
+++ b/Engine/Plugins/VirtualProduction/CameraCalibrationCore/Shaders/Private/DistortionSTMapProcessor.usf
@@ -0,0 +1,64 @@
+// Copyright Epic Games, Inc. All Rights Reserved.
+
+/*=============================================================================
+	DistortionSTMapProcessor.usf
+=============================================================================*/
+
+#include "/Engine/Public/Platform.ush"
+
+float2 TexelSize;
+uint2 TextureSize;
+Texture2D DistortionSTMap;
+SamplerState DistortionSTMapSampler;
+
+RWBuffer<float2> OutDistortedUV; 
+RWTexture2D<float2> OutUndistortionDisplacementMap;
+RWTexture2D<float2> OutDistortionDisplacementMap;
+
+
+[numthreads(8,8,1)]
+void DistortionSTMapProcessorMainCS(uint3 ThreadId : SV_DispatchThreadID)
+{
+	float2 UV = TexelSize * (ThreadId.xy + 0.5f);
+	float4 UVMapping = DistortionSTMap.SampleLevel(DistortionSTMapSampler, UV, 0);
+
+	float2 Undistortion = float2(UVMapping.x - UV.x, (1.0f - UVMapping.y) - UV.y);
+	float2 Distortion = float2(UVMapping.z - UV.x, (1.0f - UVMapping.w) - UV.y);
+
+	OutUndistortionDisplacementMap[ThreadId.xy] = Undistortion;
+	OutDistortionDisplacementMap[ThreadId.xy] = Distortion;
+
+	//Stamp distorted UV for known undistorted points
+	if(all(ThreadId.xy == uint2(0,0))) //0.0f, 0.0f
+	{
+		OutDistortedUV[0] = UVMapping.zw;
+	}
+	else if(all(ThreadId.xy == uint2(TextureSize.x * 0.5f, 0))) //0.5f, 0.0f
+	{
+		OutDistortedUV[1] = UVMapping.zw;
+	}
+	else if(all(ThreadId.xy == uint2(TextureSize.x - 1, 0))) //1.0f, 0.0f
+	{
+		OutDistortedUV[2] = UVMapping.zw;
+	}
+	else if(all(ThreadId.xy == uint2(TextureSize.x - 1, TextureSize.y * 0.5f))) //1.0f, 0.5f
+	{
+		OutDistortedUV[3] = UVMapping.zw;
+	}
+	else if(all(ThreadId.xy == uint2(TextureSize.x - 1, TextureSize.y - 1))) //1.0f, 1.0f
+	{
+		OutDistortedUV[4] = UVMapping.zw;
+	}
+	else if(all(ThreadId.xy == uint2(TextureSize.x * 0.5f, TextureSize.y - 1))) //0.5f, 1.0f
+	{
+		OutDistortedUV[5] = UVMapping.zw;
+	}
+	else if(all(ThreadId.xy == uint2(0, TextureSize.y - 1))) //0.0f, 1.0f
+	{
+		OutDistortedUV[6] = UVMapping.zw;
+	}
+	else if(all(ThreadId.xy == uint2(0, TextureSize.y * 0.5f))) //0.0f, 0.5f
+	{
+		OutDistortedUV[7] = UVMapping.zw;
+	}	
+}
diff --git a/Engine/Plugins/VirtualProduction/CameraCalibrationCore/Source/CameraCalibrationCore/CameraCalibrationCore.Build.cs b/Engine/Plugins/VirtualProduction/CameraCalibrationCore/Source/CameraCalibrationCore/CameraCalibrationCore.Build.cs
new file mode 100644
--- /dev/null
+++ b/Engine/Plugins/VirtualProduction/CameraCalibrationCore/Source/CameraCalibrationCore/CameraCalibrationCore.Build.cs
@@ -0,0 +1,46 @@
+// Copyright Epic Games, Inc. All Rights Reserved.
+
+namespace UnrealBuildTool.Rules
+{
+	public class CameraCalibrationCore : ModuleRules
+	{
+		public CameraCalibrationCore(ReadOnlyTargetRules Target) : base(Target)
+		{
+			PublicDependencyModuleNames.AddRange(
+				new string[]
+				{
+					"CinematicCamera",
+					"Core",
+					"CoreUObject",
+					"DeveloperSettings",
+					"Engine",
+					"RenderCore",
+					"RHI",
+					"Slate",
+					"SlateCore",
+				}
+			);
+
+			PrivateDependencyModuleNames.AddRange(
+				new string[]
+                {
+					"LiveLinkComponents",
+                    "Projects",
+				}
+			);
+
+			if (Target.bBuildEditor == true)
+			{
+				PrivateDependencyModuleNames.Add("UnrealEd");
+			}
+
+			PrivateIncludePaths.AddRange(
+				new string[] 
+				{
+            		//required for ScreenPass
+            		"../../../../Source/Runtime/Renderer/Private",
+            	}
+			);
+		}
+	}
+}
diff --git a/Engine/Plugins/VirtualProduction/CameraCalibrationCore/Source/CameraCalibrationCore/Private/CalibratedMapProcessor.cpp b/Engine/Plugins/VirtualProduction/CameraCalibrationCore/Source/CameraCalibrationCore/Private/CalibratedMapProcessor.cpp
new file mode 100644
--- /dev/null
+++ b/Engine/Plugins/VirtualProduction/CameraCalibrationCore/Source/CameraCalibrationCore/Private/CalibratedMapProcessor.cpp
@@ -0,0 +1,270 @@
+// Copyright Epic Games, Inc. All Rights Reserved.
+
+#include "CalibratedMapProcessor.h"
+
+#include "Engine/Texture.h"
+#include "Engine/TextureRenderTarget2D.h"
+#include "GlobalShader.h"		
+#include "CameraCalibrationCoreLog.h"
+#include "PixelShaderUtils.h"
+#include "RenderGraphUtils.h"
+#include "ShaderParameterStruct.h"
+#include "ScreenPass.h"
+
+
+
+class FCalibratedMapDerivedDataCS : public FGlobalShader
+{
+public:
+	DECLARE_GLOBAL_SHADER(FCalibratedMapDerivedDataCS);
+
+	SHADER_USE_PARAMETER_STRUCT(FCalibratedMapDerivedDataCS, FGlobalShader);
+
+	BEGIN_SHADER_PARAMETER_STRUCT(FParameters, )
+		SHADER_PARAMETER(FVector2D, TexelSize)
+		SHADER_PARAMETER(FIntPoint, TextureSize)
+		
+		SHADER_PARAMETER_RDG_TEXTURE(Texture2D, DistortionSTMap)
+		SHADER_PARAMETER_SAMPLER(SamplerState, DistortionSTMapSampler)
+
+		SHADER_PARAMETER_RDG_BUFFER_UAV(RWBuffer<FVector2D>, OutDistortedUV)
+		SHADER_PARAMETER_RDG_TEXTURE_UAV(RWTexture2D<float2>, OutUndistortionDisplacementMap)
+		SHADER_PARAMETER_RDG_TEXTURE_UAV(RWTexture2D<float2>, OutDistortionDisplacementMap)
+		END_SHADER_PARAMETER_STRUCT()
+
+public:
+	// Called by the engine to determine which permutations to compile for this shader
+	static bool ShouldCompilePermutation(const FGlobalShaderPermutationParameters& Parameters)
+	{
+		return IsFeatureLevelSupported(Parameters.Platform, ERHIFeatureLevel::SM5);
+	}
+};
+
+ // --------------------------------------------------------------------------------------------------------------------
+IMPLEMENT_GLOBAL_SHADER(FCalibratedMapDerivedDataCS, "/Plugin/CameraCalibrationCore/Private/DistortionSTMapProcessor.usf", "DistortionSTMapProcessorMainCS", SF_Compute);
+							 
+
+FCalibratedMapProcessor::FCalibratedMapProcessor()
+{
+	RunningJobs.Reserve(JobCount);
+	for(int32 Index = 0; Index < JobCount; ++Index)
+	{
+		TSharedPtr<FDerivedDistortionDataJob> NewJob = MakeShared<FDerivedDistortionDataJob>();
+		NewJob->Readback = MakePimpl<FRHIGPUBufferReadback>(*FString::Printf(TEXT("DerivedDistortionDataJobReadback_%02d"), Index));
+		AvailableJobs.Enqueue(NewJob);
+	}
+}
+
+FCalibratedMapProcessor::~FCalibratedMapProcessor()
+{
+	FlushRenderingCommands();
+
+	FScopeLock Lock(&RunningJobsCriticalSection);
+	RunningJobs.Empty();
+	CompletedJobs.Empty();
+	AvailableJobs.Empty();
+}
+
+void FCalibratedMapProcessor::Update()
+{
+	//If we have completed jobs, pop them and callback owner
+	TSharedPtr<FDerivedDistortionDataJob> CompletedJob = nullptr;
+	while (CompletedJobs.Dequeue(CompletedJob))
+	{
+		CompletedJob->JobArgs.JobCompletedCallback.ExecuteIfBound(CompletedJob->Output);
+		
+		AvailableJobs.Enqueue(CompletedJob);
+	}
+
+	//Fill in available job slots with any pending ones
+	while (PendingJobs.Peek())
+	{
+		//push pending jobs
+		TSharedPtr<FDerivedDistortionDataJob> NewJob = nullptr;
+		if (AvailableJobs.Dequeue(NewJob))
+		{
+			//We have a slot available
+			FDerivedDistortionDataJobArgs PendingJob;
+
+			//Consume Pending job that was peeked
+			check(PendingJobs.Dequeue(PendingJob));
+
+			NewJob->Output.Reset();
+			NewJob->Output.Identifier = PendingJob.Identifier;
+			NewJob->JobArgs = MoveTemp(PendingJob);
+			ExecuteJob(MoveTemp(NewJob));
+		}
+		else
+		{
+			//No more job slots available, exit
+			break;
+		}
+	}
+
+	ENQUEUE_RENDER_COMMAND(FCalibratedMapProcessor_Update_RenderThread)(
+		[this](FRHICommandListImmediate& RHICmdList)
+		{
+			Update_RenderThread();
+		});
+}
+
+bool FCalibratedMapProcessor::PushDerivedDistortionDataJob(FDerivedDistortionDataJobArgs&& JobArgs)
+{
+	//Validate arguments before pushing the job
+	if (!JobArgs.JobCompletedCallback.IsBound()
+		|| !JobArgs.SourceDistortionMap.IsValid()
+		|| !JobArgs.OutputUndistortionDisplacementMap.IsValid()
+		|| !JobArgs.OutputDistortionDisplacementMap.IsValid())
+	{
+		return false;
+	}
+
+	TSharedPtr<FDerivedDistortionDataJob> NewJob = nullptr;
+	if (AvailableJobs.Dequeue(NewJob))
+	{
+		NewJob->Output.Reset();
+		NewJob->Output.Identifier = JobArgs.Identifier;
+		NewJob->JobArgs = MoveTemp(JobArgs);
+		ExecuteJob(NewJob);
+	}
+	
+	if(NewJob == nullptr)
+	{
+		//No slots available, push that to the pending job
+		PendingJobs.Enqueue(MoveTemp(JobArgs));
+	}
+	
+	return true;
+}
+
+void FCalibratedMapProcessor::ExecuteJob(TSharedPtr<FDerivedDistortionDataJob> Job)
+{
+	check(Job);
+	{
+		//Early validation of inputs in case some textures are now invalid
+		UTexture* SourceTexture = Job->JobArgs.SourceDistortionMap.Get();
+		UTextureRenderTarget2D* DestinationUndistortionTexture = Job->JobArgs.OutputUndistortionDisplacementMap.Get();
+		UTextureRenderTarget2D* DestinationDistortionTexture = Job->JobArgs.OutputDistortionDisplacementMap.Get();
+		if(SourceTexture == nullptr
+		|| DestinationUndistortionTexture == nullptr
+		|| DestinationDistortionTexture == nullptr
+		|| SourceTexture->Resource == nullptr
+		|| DestinationUndistortionTexture->Resource == nullptr
+		|| DestinationDistortionTexture->Resource == nullptr)
+		{
+			Job->State = EDerivedDistortionDataJobState::Completed;
+			Job->Output.Result = EDerivedDistortionDataResult::Error;
+			CompletedJobs.Enqueue(Job);
+			return;
+		}
+		
+		FTextureResource* SourceDistortionMap = SourceTexture->Resource;
+		FTextureResource* DestinationUndistortionDisplacementMap = DestinationUndistortionTexture->Resource;
+		FTextureResource* DestinationDistortionDisplacementMap = DestinationDistortionTexture->Resource;
+
+		{
+			FScopeLock Lock(&RunningJobsCriticalSection);
+			RunningJobs.Add(Job);
+		}
+		
+		ENQUEUE_RENDER_COMMAND(FCalibratedMapProcessor_ComputeDerivedDistortionData)(
+		[DestinationUndistortionDisplacementMap, DestinationDistortionDisplacementMap, SourceDistortionMap, Job](FRHICommandListImmediate& RHICmdList)
+        {
+            FRDGBuilder GraphBuilder(RHICmdList);
+
+			FRDGTextureRef SourceSTMap = GraphBuilder.RegisterExternalTexture(CreateRenderTarget(SourceDistortionMap->TextureRHI, TEXT("SourceSTMap")));
+			FRDGTextureRef UndistortionDisplacementMap = GraphBuilder.RegisterExternalTexture(CreateRenderTarget(DestinationUndistortionDisplacementMap->TextureRHI, TEXT("UndistortionDestinationDisplacementMap")));
+			FRDGTextureRef DistortionDisplacementMap = GraphBuilder.RegisterExternalTexture(CreateRenderTarget(DestinationDistortionDisplacementMap->TextureRHI, TEXT("DistortionDestinationDisplacementMap")));
+
+			const int32 EntriesCount = Job->Output.EdgePointCount;
+			const int32 ReadbackDataByteCount = Job->Output.EdgePointsDistortedUVs.GetTypeSize();
+			FRDGBufferDesc EdgePointsDistortedUVDesc = FRDGBufferDesc::CreateBufferDesc(ReadbackDataByteCount, EntriesCount);
+			EdgePointsDistortedUVDesc.Usage = EBufferUsageFlags(EdgePointsDistortedUVDesc.Usage | BUF_SourceCopy);
+			FRDGBufferRef EdgePointsBuffer = GraphBuilder.CreateBuffer(EdgePointsDistortedUVDesc, TEXT("EdgePointsBuffer"));
+           	
+			FCalibratedMapDerivedDataCS::FParameters* Parameters = GraphBuilder.AllocParameters<FCalibratedMapDerivedDataCS::FParameters>();
+            Parameters->DistortionSTMap = SourceSTMap;
+           	Parameters->DistortionSTMapSampler = TStaticSamplerState<SF_Bilinear, AM_Clamp, AM_Clamp, AM_Clamp>::GetRHI();
+           	Parameters->TexelSize = FVector2D(1.0f / UndistortionDisplacementMap->Desc.Extent.X, 1.0f / UndistortionDisplacementMap->Desc.Extent.Y);
+            Parameters->TextureSize = FIntPoint(UndistortionDisplacementMap->Desc.Extent.X, UndistortionDisplacementMap->Desc.Extent.Y);
+
+            //Create UAVs for compute shader outputs
+           	Parameters->OutDistortedUV = GraphBuilder.CreateUAV(FRDGBufferUAVDesc(EdgePointsBuffer, PF_G32R32F));
+			Parameters->OutUndistortionDisplacementMap = GraphBuilder.CreateUAV(UndistortionDisplacementMap);
+			Parameters->OutDistortionDisplacementMap = GraphBuilder.CreateUAV(DistortionDisplacementMap);
+
+            FGlobalShaderMap* GlobalShaderMap = GetGlobalShaderMap(GMaxRHIFeatureLevel);
+			TShaderMapRef<FCalibratedMapDerivedDataCS> ComputeShader(GlobalShaderMap);
+
+            GraphBuilder.AddPass(
+            			RDG_EVENT_NAME("LensFileSTMapConversion"),
+            			Parameters,
+            			ERDGPassFlags::Compute,
+            			[Parameters, ComputeShader](FRHICommandList& RHICmdList)
+						{
+            				const FIntVector GroupCount = FComputeShaderUtils::GetGroupCount(Parameters->TextureSize, FComputeShaderUtils::kGolden2DGroupSize);
+            				FComputeShaderUtils::Dispatch(RHICmdList, ComputeShader, *Parameters, GroupCount);
+						}
+            );
+
+            //We need to access the data related to points on the edge to compute overscan factor
+			AddEnqueueCopyPass(GraphBuilder, Job->Readback.Get(), EdgePointsBuffer, ReadbackDataByteCount * EntriesCount);
+            
+			//Kick graph
+			GraphBuilder.Execute();
+
+			//Once graph is kicked, mark the job as waiting
+			Job->State = EDerivedDistortionDataJobState::AwaitingResult;
+        });		
+	}
+}
+
+
+
+void FCalibratedMapProcessor::Update_RenderThread()
+{
+	TArray<TSharedPtr<FDerivedDistortionDataJob>, TInlineAllocator<JobCount>> ToComplete;
+	{
+		FScopeLock Lock(&RunningJobsCriticalSection);
+		for (auto It = RunningJobs.CreateIterator(); It; ++It)
+		{
+			TSharedPtr<FDerivedDistortionDataJob> CompletedJob = *It;
+			if (CompletedJob->State == EDerivedDistortionDataJobState::AwaitingResult)
+			{
+				if (CompletedJob->Readback->IsReady())
+				{
+					ToComplete.Add(CompletedJob);
+					It.RemoveCurrent();
+				}
+			}
+		}
+	}
+
+	//Now process out of critical section completed jobs to add them to the queue
+	for(TSharedPtr<FDerivedDistortionDataJob>& CompletedJob : ToComplete)
+	{
+		constexpr int32 EntriesCount = FDerivedDistortionDataJobOutput::EdgePointCount;
+		const int32 ReadbackDataByteCount = CompletedJob->Output.EdgePointsDistortedUVs.GetTypeSize();
+		const FVector2D* EdgePointsDistortedUV = static_cast<const FVector2D*>(CompletedJob->Readback->Lock(EntriesCount * ReadbackDataByteCount));
+
+		if (EdgePointsDistortedUV)
+		{
+			for (int32 Index = 0; Index < CompletedJob->Output.EdgePointsDistortedUVs.Num(); ++Index)
+			{
+				CompletedJob->Output.EdgePointsDistortedUVs[Index] = EdgePointsDistortedUV[Index];
+				CompletedJob->Output.Result = EDerivedDistortionDataResult::Success;
+			}
+		}
+		else
+		{
+			UE_LOG(LogCameraCalibrationCore, Error, TEXT("Failed to retrieve edge points distorted UVs"));
+			CompletedJob->Output.Result = EDerivedDistortionDataResult::Error;
+		}
+
+		CompletedJob->Readback->Unlock();
+
+		CompletedJob->State = EDerivedDistortionDataJobState::Completed;
+		CompletedJobs.Enqueue(CompletedJob);
+	}
+}
+
diff --git a/Engine/Plugins/VirtualProduction/CameraCalibrationCore/Source/CameraCalibrationCore/Private/CalibratedMapProcessor.h b/Engine/Plugins/VirtualProduction/CameraCalibrationCore/Source/CameraCalibrationCore/Private/CalibratedMapProcessor.h
new file mode 100644
--- /dev/null
+++ b/Engine/Plugins/VirtualProduction/CameraCalibrationCore/Source/CameraCalibrationCore/Private/CalibratedMapProcessor.h
@@ -0,0 +1,142 @@
+// Copyright Epic Games, Inc. All Rights Reserved.
+
+#pragma once
+
+#include "ICalibratedMapProcessor.h"
+#include "Containers/Queue.h"
+#include "Templates/PimplPtr.h"
+#include "UObject/WeakObjectPtr.h"
+
+class UTexture;
+struct FDerivedDistortionData;
+class FRHIGPUBufferReadback;
+class UTextureRenderTarget2D;
+
+/** 
+ * Result of processing derived data 
+ */
+enum class EDerivedDistortionDataResult : uint8
+{
+	Success,
+	Error,
+};
+
+/**  
+ * Arguments for a job to derived data from calibrated map 
+ */
+struct FDerivedDistortionDataJobArgs
+{
+	/** Identifier of the source/destination point */
+	FGuid Identifier;
+
+	/** Source calibrated map */
+	//When coefficient route is supported, we might look at how we can have a mutual interface
+	TWeakObjectPtr<UTexture> SourceDistortionMap;
+
+	/** Output derived undistortion displacement map */
+	TWeakObjectPtr<UTextureRenderTarget2D> OutputUndistortionDisplacementMap;
+
+	/** Output derived distortion displacement map */
+	TWeakObjectPtr<UTextureRenderTarget2D> OutputDistortionDisplacementMap;
+
+	/** Delegate to be called when job is completed */
+	FOnDerivedDistortionJobCompleted JobCompletedCallback;
+};
+
+/** 
+ * Derived data that was generated 
+ */
+struct FDerivedDistortionDataJobOutput
+{
+	FDerivedDistortionDataJobOutput()
+	{
+		Reset();
+	}
+	
+	/** Used when reusing a job */
+	void Reset()
+	{
+		Result = EDerivedDistortionDataResult::Error;
+		EdgePointsDistortedUVs.SetNumZeroed(EdgePointCount);
+	}
+	
+	/** Result of that job */
+	EDerivedDistortionDataResult Result = EDerivedDistortionDataResult::Error;
+
+	/** Identifier that we were given in job arguments */
+	FGuid Identifier;
+
+	/** Distorted UVs we found at the 8 edge points */
+	static constexpr int32 EdgePointCount = 8;
+	TArray<FVector2D> EdgePointsDistortedUVs;
+};
+
+/** State of a job */
+enum class EDerivedDistortionDataJobState : uint8
+{
+	Ready,
+	AwaitingResult,
+	Completed,
+};
+
+/**
+ * Info for a given job to derive data from calibrated map
+ */
+struct FDerivedDistortionDataJob
+{
+	/** State of the job */
+	EDerivedDistortionDataJobState State = EDerivedDistortionDataJobState::Ready;
+	
+	/** Arguments for this job */
+	FDerivedDistortionDataJobArgs JobArgs;
+	
+	/** Manages data we read back from compute shader operation */
+	TPimplPtr<FRHIGPUBufferReadback> Readback;
+
+	/** Output generated by this job */
+	FDerivedDistortionDataJobOutput Output;
+};
+
+/**
+ * Processor used to launch job to derived data out of calibrated
+ * STMaps/UVMaps. Goal is to have an intermediate data model
+ * that we use when we evaluate the LensFile
+ */
+class FCalibratedMapProcessor : public ICalibratedMapProcessor
+{
+public:
+
+	FCalibratedMapProcessor();
+	virtual ~FCalibratedMapProcessor() override;
+
+	//~ Begin ICalibratedMapProcessor interface
+	virtual void Update() override;
+	virtual bool PushDerivedDistortionDataJob(FDerivedDistortionDataJobArgs&& JobArgs) override;
+	//~ End ICalibratedMapProcessor interface
+
+private:
+
+	/** Execute job by queuing on render thread */
+	void ExecuteJob(TSharedPtr<FDerivedDistortionDataJob> Job);
+	
+	/** Verify if any running jobs are done doing readback and convert to completed */
+	void Update_RenderThread();
+	
+private:
+
+	/** Number of jobs slots to create. Limiting the amount here to limit number of allocation. */
+	static constexpr int32 JobCount = 4;
+	
+	/** Job slots available to be configured and launched */
+	TQueue<TSharedPtr<FDerivedDistortionDataJob>> AvailableJobs;
+
+	/** Jobs currently being ran and waiting on gpu readback to be available */
+	FCriticalSection RunningJobsCriticalSection;
+	TArray<TSharedPtr<FDerivedDistortionDataJob>> RunningJobs;
+
+	/** Jobs that were completed, ready to be dispatched */
+	TQueue<TSharedPtr<FDerivedDistortionDataJob>> CompletedJobs;
+
+	/** Jobs that were requested but no slots were available */
+	TQueue<FDerivedDistortionDataJobArgs> PendingJobs;
+};
diff --git a/Engine/Plugins/VirtualProduction/CameraCalibrationCore/Source/CameraCalibrationCore/Private/CameraCalibrationCoreLog.h b/Engine/Plugins/VirtualProduction/CameraCalibrationCore/Source/CameraCalibrationCore/Private/CameraCalibrationCoreLog.h
new file mode 100644
--- /dev/null
+++ b/Engine/Plugins/VirtualProduction/CameraCalibrationCore/Source/CameraCalibrationCore/Private/CameraCalibrationCoreLog.h
@@ -0,0 +1,8 @@
+// Copyright Epic Games, Inc. All Rights Reserved.
+
+#pragma once
+
+#include "CoreMinimal.h"
+#include "Logging/LogMacros.h"
+
+DECLARE_LOG_CATEGORY_EXTERN(LogCameraCalibrationCore, Log, All);
diff --git a/Engine/Plugins/VirtualProduction/CameraCalibrationCore/Source/CameraCalibrationCore/Private/CameraCalibrationCoreModule.cpp b/Engine/Plugins/VirtualProduction/CameraCalibrationCore/Source/CameraCalibrationCore/Private/CameraCalibrationCoreModule.cpp
new file mode 100644
--- /dev/null
+++ b/Engine/Plugins/VirtualProduction/CameraCalibrationCore/Source/CameraCalibrationCore/Private/CameraCalibrationCoreModule.cpp
@@ -0,0 +1,150 @@
+// Copyright Epic Games, Inc. All Rights Reserved.
+
+#include "CameraCalibrationCoreModule.h"
+
+#include "CameraCalibrationCoreLog.h"
+#include "CameraCalibrationSettings.h"
+#include "CameraCalibrationSubsystem.h"
+#include "Engine/Engine.h"
+#include "HAL/IConsoleManager.h"
+#include "Interfaces/IPluginManager.h"
+#include "LensFile.h"
+#include "Misc/App.h"
+#include "Misc/CoreDelegates.h"
+#include "Misc/Paths.h"
+#include "Models/SphericalLensModel.h"
+#include "Modules/ModuleManager.h"
+#include "ShaderCore.h"
+#include "UObject/UObjectGlobals.h"
+
+DEFINE_LOG_CATEGORY(LogCameraCalibrationCore);
+
+
+static TAutoConsoleVariable<FString> CVarStartupLensFile(
+	TEXT("CameraCalibration.StartupLensFile"),
+	TEXT(""),
+	TEXT("Startup Lens File\n"),
+	ECVF_ReadOnly
+);
+
+
+void FCameraCalibrationCoreModule::StartupModule()
+{
+	FString PluginShaderDir = FPaths::Combine(IPluginManager::Get().FindPlugin(TEXT("CameraCalibrationCore"))->GetBaseDir(), TEXT("Shaders"));
+	AddShaderSourceDirectoryMapping(TEXT("/Plugin/CameraCalibrationCore"), PluginShaderDir);
+
+	ApplyStartupLensFile();
+	RegisterDistortionModels();
+}
+
+void FCameraCalibrationCoreModule::ShutdownModule()
+{
+	UnregisterDistortionModels();
+
+	if (PostEngineInitHandle.IsValid())
+	{
+		FCoreDelegates::OnPostEngineInit.Remove(PostEngineInitHandle);
+	}
+}
+
+void FCameraCalibrationCoreModule::ApplyStartupLensFile()
+{
+	auto ApplyLensFile = [this]()
+	{
+		ULensFile* StartupLensFile = nullptr;
+
+		// Try to load from CVar
+		{
+			const FString LensFileName = CVarStartupLensFile.GetValueOnGameThread();
+
+			if (LensFileName.Len())
+			{
+				if (UObject* Object = StaticLoadObject(ULensFile::StaticClass(), nullptr, *LensFileName))
+				{
+					StartupLensFile = CastChecked<ULensFile>(Object);
+				}
+
+				if (StartupLensFile)
+				{
+					UE_LOG(LogCameraCalibrationCore, Display, TEXT("Loading Lens File specified in CVar CameraCalibration.StartupLensFile: '%s'"), *LensFileName);
+				}
+			}
+		}
+
+#if WITH_EDITOR
+		// Try to load user settings
+		if (StartupLensFile == nullptr)
+		{
+			StartupLensFile = GetDefault<UCameraCalibrationEditorSettings>()->GetUserLensFile();
+			if (StartupLensFile)
+			{
+				UE_LOG(LogCameraCalibrationCore, Display, TEXT("Loading Lens File specified in user settings: '%s'"), *StartupLensFile->GetName());
+			}
+		}
+#endif
+
+		// Try loading default lens file from project settings
+		if (StartupLensFile == nullptr)
+		{
+			StartupLensFile = GetDefault<UCameraCalibrationSettings>()->GetStartupLensFile();
+
+			if (StartupLensFile)
+			{
+				UE_LOG(LogCameraCalibrationCore, Display, TEXT("Loading Lens File specified in project settings: '%s'"), *StartupLensFile->GetName());
+			}
+		}
+
+		UCameraCalibrationSubsystem* SubSystem = GEngine->GetEngineSubsystem<UCameraCalibrationSubsystem>();
+		check(SubSystem);
+		SubSystem->SetDefaultLensFile(StartupLensFile);
+	};
+
+	if (FApp::CanEverRender())
+	{
+		if (GEngine && GEngine->IsInitialized())
+		{
+			ApplyLensFile();
+		}
+		else
+		{
+			PostEngineInitHandle = FCoreDelegates::OnPostEngineInit.AddLambda(ApplyLensFile);
+		}
+	}
+}
+
+void FCameraCalibrationCoreModule::RegisterDistortionModels()
+{
+	auto RegisterModels = [this]()
+	{
+		// Register all lens models defined in this module
+		UCameraCalibrationSubsystem* SubSystem = GEngine->GetEngineSubsystem<UCameraCalibrationSubsystem>();
+		SubSystem->RegisterDistortionModel(USphericalLensModel::StaticClass());
+	};
+
+	if (FApp::CanEverRender())
+	{
+		if (GEngine && GEngine->IsInitialized())
+		{
+			RegisterModels();
+		}
+		else
+		{
+			PostEngineInitHandle = FCoreDelegates::OnPostEngineInit.AddLambda(RegisterModels);
+		}
+	}
+}
+
+void FCameraCalibrationCoreModule::UnregisterDistortionModels()
+{
+	if (GEngine)
+	{
+		// Unregister all lens models defined in this module
+		if (UCameraCalibrationSubsystem* SubSystem = GEngine->GetEngineSubsystem<UCameraCalibrationSubsystem>())
+		{
+			SubSystem->UnregisterDistortionModel(USphericalLensModel::StaticClass());
+		}
+	}
+}
+
+
+IMPLEMENT_MODULE(FCameraCalibrationCoreModule, CameraCalibrationCore)
\ No newline at end of file
diff --git a/Engine/Plugins/VirtualProduction/CameraCalibrationCore/Source/CameraCalibrationCore/Private/CameraCalibrationCoreModule.h b/Engine/Plugins/VirtualProduction/CameraCalibrationCore/Source/CameraCalibrationCore/Private/CameraCalibrationCoreModule.h
new file mode 100644
--- /dev/null
+++ b/Engine/Plugins/VirtualProduction/CameraCalibrationCore/Source/CameraCalibrationCore/Private/CameraCalibrationCoreModule.h
@@ -0,0 +1,26 @@
+// Copyright Epic Games, Inc. All Rights Reserved.
+
+#pragma once
+
+#include "Modules/ModuleInterface.h"
+#include "Delegates/IDelegateInstance.h"
+
+
+class FCameraCalibrationCoreModule : public IModuleInterface
+{
+public:
+
+	//~ Begin IModuleInterface interface
+	virtual void StartupModule() override;
+	virtual void ShutdownModule() override;
+	//~ End IModuleInterface interface
+
+private:
+	void ApplyStartupLensFile();
+	void RegisterDistortionModels();
+	void UnregisterDistortionModels();
+
+private:
+
+	FDelegateHandle PostEngineInitHandle;
+};
diff --git a/Engine/Plugins/VirtualProduction/CameraCalibrationCore/Source/CameraCalibrationCore/Private/CameraCalibrationSettings.cpp b/Engine/Plugins/VirtualProduction/CameraCalibrationCore/Source/CameraCalibrationCore/Private/CameraCalibrationSettings.cpp
new file mode 100644
--- /dev/null
+++ b/Engine/Plugins/VirtualProduction/CameraCalibrationCore/Source/CameraCalibrationCore/Private/CameraCalibrationSettings.cpp
@@ -0,0 +1,89 @@
+// Copyright Epic Games, Inc. All Rights Reserved.
+
+#include "CameraCalibrationSettings.h"
+
+#include "SphericalLensDistortionModelHandler.h"
+
+UCameraCalibrationSettings::UCameraCalibrationSettings()
+{
+	DefaultUndistortionDisplacementMaterials.Add(USphericalLensDistortionModelHandler::StaticClass(), TSoftObjectPtr<UMaterialInterface>(FSoftObjectPath(TEXT("/CameraCalibrationCore/Materials/M_SphericalUndistortionDisplacementMap.M_SphericalUndistortionDisplacementMap"))));
+	DefaultDistortionDisplacementMaterials.Add(USphericalLensDistortionModelHandler::StaticClass(), TSoftObjectPtr<UMaterialInterface>(FSoftObjectPath(TEXT("/CameraCalibrationCore/Materials/M_SphericalDistortionDisplacementMap.M_SphericalDistortionDisplacementMap"))));
+	DefaultDistortionMaterials.Add(USphericalLensDistortionModelHandler::StaticClass(), TSoftObjectPtr<UMaterialInterface>(FSoftObjectPath(TEXT("/CameraCalibrationCore/Materials/M_DistortionPostProcess.M_DistortionPostProcess"))));
+}
+
+FName UCameraCalibrationSettings::GetCategoryName() const
+{
+	return TEXT("Plugins");
+}
+
+#if WITH_EDITOR
+FText UCameraCalibrationSettings::GetSectionText() const
+{
+	return NSLOCTEXT("CameraCalibrationPlugin", "CameraCalibrationSettingsSection", "Camera Calibration");
+}
+
+#endif
+
+ULensFile* UCameraCalibrationSettings::GetStartupLensFile() const
+{
+	return StartupLensFile.LoadSynchronous();
+}
+
+UMaterialInterface* UCameraCalibrationSettings::GetDefaultUndistortionDisplacementMaterial(const TSubclassOf<ULensDistortionModelHandlerBase>& InModelHandler) const
+{
+	const TSoftObjectPtr<UMaterialInterface>* UndistortionDisplacementMaterial = DefaultUndistortionDisplacementMaterials.Find(InModelHandler);
+	if (!UndistortionDisplacementMaterial)
+	{
+		return nullptr;
+	}
+	return UndistortionDisplacementMaterial->LoadSynchronous();
+}
+
+UMaterialInterface* UCameraCalibrationSettings::GetDefaultDistortionDisplacementMaterial(const TSubclassOf<ULensDistortionModelHandlerBase>& InModelHandler) const
+{
+	const TSoftObjectPtr<UMaterialInterface>* DistortionDisplacementMaterial = DefaultDistortionDisplacementMaterials.Find(InModelHandler);
+	if (!DistortionDisplacementMaterial)
+	{
+		return nullptr;
+	}
+	return DistortionDisplacementMaterial->LoadSynchronous();
+}
+
+UMaterialInterface* UCameraCalibrationSettings::GetDefaultDistortionMaterial(const TSubclassOf<ULensDistortionModelHandlerBase>& InModelHandler) const
+{
+	const TSoftObjectPtr<UMaterialInterface>* DistortionMaterial = DefaultDistortionMaterials.Find(InModelHandler);
+	if (!DistortionMaterial)
+	{
+		return nullptr;
+	}
+	return DistortionMaterial->LoadSynchronous();
+}
+
+FName UCameraCalibrationEditorSettings::GetCategoryName() const
+{
+	return TEXT("Plugins");
+}
+
+#if WITH_EDITOR
+FText UCameraCalibrationEditorSettings::GetSectionText() const
+{
+	return NSLOCTEXT("CameraCalibrationEditorPlugin", "CameraCalibrationEditorSettingsSection", "Camera Calibration");
+}
+#endif
+
+ULensFile* UCameraCalibrationEditorSettings::GetUserLensFile() const
+{
+#if WITH_EDITOR
+	return UserLensFile.LoadSynchronous();
+#else
+	return nullptr;
+#endif
+}
+
+void UCameraCalibrationEditorSettings::SetUserLensFile(ULensFile* InLensFile)
+{
+#if WITH_EDITOR
+	UserLensFile = InLensFile;
+	SaveConfig();
+#endif
+}
diff --git a/Engine/Plugins/VirtualProduction/CameraCalibrationCore/Source/CameraCalibrationCore/Private/CameraCalibrationSubsystem.cpp b/Engine/Plugins/VirtualProduction/CameraCalibrationCore/Source/CameraCalibrationCore/Private/CameraCalibrationSubsystem.cpp
new file mode 100644
--- /dev/null
+++ b/Engine/Plugins/VirtualProduction/CameraCalibrationCore/Source/CameraCalibrationCore/Private/CameraCalibrationSubsystem.cpp
@@ -0,0 +1,238 @@
+// Copyright Epic Games, Inc. All Rights Reserved.
+
+#include "CameraCalibrationSubsystem.h"
+
+#include "CameraCalibrationCoreLog.h"
+#include "CineCameraComponent.h"
+#include "Engine/TimecodeProvider.h"
+#include "Misc/CoreDelegates.h"
+#include "UObject/UObjectIterator.h"
+
+#if WITH_EDITOR
+#include "Editor.h"
+#endif // WITH_EDITOR
+
+#define LOCTEXT_NAMESPACE "CameraCalibrationSubsystem"
+
+
+
+ULensFile* UCameraCalibrationSubsystem::GetDefaultLensFile() const
+{
+	return DefaultLensFile;
+}
+
+void UCameraCalibrationSubsystem::SetDefaultLensFile(ULensFile* NewDefaultLensFile)
+{
+	//Todo : Add callbacks when default lens file changes
+	DefaultLensFile = NewDefaultLensFile;
+}
+
+ULensFile* UCameraCalibrationSubsystem::GetLensFile(const FLensFilePicker& Picker) const
+{
+	ULensFile* ReturnedLens = nullptr;
+
+	if (Picker.bOverrideDefaultLensFile)
+	{
+		ReturnedLens = Picker.LensFile;
+	}
+	else
+	{
+		ReturnedLens = GetDefaultLensFile();
+	}
+
+	return ReturnedLens;
+}
+
+TArray<ULensDistortionModelHandlerBase*> UCameraCalibrationSubsystem::GetDistortionModelHandlers(UCineCameraComponent* Component)
+{
+	TArray<ULensDistortionModelHandlerBase*> Handlers;
+	LensDistortionHandlerMap.MultiFind(Component, Handlers);
+
+	return Handlers;
+}
+
+ULensDistortionModelHandlerBase* UCameraCalibrationSubsystem::FindDistortionModelHandler(FDistortionHandlerPicker& DistortionHandlerPicker, bool bUpdatePicker) const
+{
+	TArray<ULensDistortionModelHandlerBase*> Handlers;
+	LensDistortionHandlerMap.MultiFind(DistortionHandlerPicker.TargetCameraComponent, Handlers);
+
+	// Look through the available handlers to find the one driven by the distortion source's producer
+	for (ULensDistortionModelHandlerBase* Handler : Handlers)
+	{
+		if (Handler->GetDistortionProducerID() == DistortionHandlerPicker.DistortionProducerID)
+		{
+			// Reassign the input handler picker's display name to that of the found handler
+			if (bUpdatePicker)
+			{
+				DistortionHandlerPicker.HandlerDisplayName = Handler->GetDisplayName();
+			}
+			return Handler;
+		}
+	}
+
+	// Look through the available handlers to find one with the same name as the input handler picker's display name
+	for (ULensDistortionModelHandlerBase* Handler : Handlers)
+	{
+		if (Handler->GetDisplayName() == DistortionHandlerPicker.HandlerDisplayName)
+		{
+			// Reassign the input handler picker's producer to that of the found handler
+			if (bUpdatePicker)
+			{
+				DistortionHandlerPicker.DistortionProducerID = Handler->GetDistortionProducerID();
+			}
+			return Handler;
+		}
+	}
+
+	return nullptr;
+}
+
+ULensDistortionModelHandlerBase* UCameraCalibrationSubsystem::FindOrCreateDistortionModelHandler(FDistortionHandlerPicker& DistortionHandlerPicker, const TSubclassOf<ULensModel> LensModelClass)
+{
+	if (LensModelClass == nullptr)
+	{
+		return nullptr;
+	}
+
+	// Attempt to find a handler associated with the input distortion source
+	if (ULensDistortionModelHandlerBase* Handler = FindDistortionModelHandler(DistortionHandlerPicker, false))
+	{
+		// If the existing handler supports the input model, simply return it
+		if (Handler->IsModelSupported(LensModelClass))
+		{
+			// The display name may have changed, even if the handler already exists, so update the display name of the existing handler to match the distortion source
+			Handler->SetDisplayName(DistortionHandlerPicker.HandlerDisplayName);
+			return Handler;
+		}
+		else
+		{
+			// If the input distortion source has an existing handler, but model does not match the input model, remove the old handler from the map
+			LensDistortionHandlerMap.Remove(DistortionHandlerPicker.TargetCameraComponent, Handler);
+		}
+	}
+
+	// If no handler exists for the input distortion source, create a new one
+	const TSubclassOf<ULensDistortionModelHandlerBase> HandlerClass = ULensModel::GetHandlerClass(LensModelClass);
+	ULensDistortionModelHandlerBase* NewHandler = nullptr;
+	if(HandlerClass)
+	{
+		NewHandler = NewObject<ULensDistortionModelHandlerBase>(DistortionHandlerPicker.TargetCameraComponent, HandlerClass);
+		NewHandler->SetDistortionProducerID(DistortionHandlerPicker.DistortionProducerID);
+		NewHandler->SetDisplayName(DistortionHandlerPicker.HandlerDisplayName);
+		LensDistortionHandlerMap.Add(FObjectKey(DistortionHandlerPicker.TargetCameraComponent), NewHandler);
+	}
+	else
+	{
+		UE_LOG(LogCameraCalibrationCore, Verbose, TEXT("Could not create DistortionHandler for LensModel '%s'"), *LensModelClass->GetName());
+	}
+
+	return NewHandler;
+}
+
+void UCameraCalibrationSubsystem::UnregisterDistortionModelHandler(UCineCameraComponent* Component, ULensDistortionModelHandlerBase* Handler)
+{
+	LensDistortionHandlerMap.Remove(Component, Handler);
+}
+
+void UCameraCalibrationSubsystem::RegisterDistortionModel(TSubclassOf<ULensModel> LensModel)
+{
+	LensModelMap.Add(LensModel->GetDefaultObject<ULensModel>()->GetModelName(), LensModel);
+}
+
+void UCameraCalibrationSubsystem::UnregisterDistortionModel(TSubclassOf<ULensModel> LensModel)
+{
+	LensModelMap.Remove(LensModel->GetDefaultObject<ULensModel>()->GetModelName());
+}
+
+TSubclassOf<ULensModel> UCameraCalibrationSubsystem::GetRegisteredLensModel(FName ModelName) const
+{
+	if (LensModelMap.Contains(ModelName))
+	{
+		return LensModelMap[ModelName];
+	}
+	return nullptr;
+}
+
+TSubclassOf<UCameraNodalOffsetAlgo> UCameraCalibrationSubsystem::GetCameraNodalOffsetAlgo(FName Name) const
+{
+	if (CameraNodalOffsetAlgosMap.Contains(Name))
+	{
+		return CameraNodalOffsetAlgosMap[Name];
+	}
+	return nullptr;
+}
+
+TArray<FName> UCameraCalibrationSubsystem::GetCameraNodalOffsetAlgos() const
+{
+	TArray<FName> OutKeys;
+	CameraNodalOffsetAlgosMap.GetKeys(OutKeys);
+	return OutKeys;
+}
+
+TSubclassOf<UCameraCalibrationStep> UCameraCalibrationSubsystem::GetCameraCalibrationStep(FName Name) const
+{
+	if (CameraCalibrationStepsMap.Contains(Name))
+	{
+		return CameraCalibrationStepsMap[Name];
+	}
+	return nullptr;
+}
+
+TArray<FName> UCameraCalibrationSubsystem::GetCameraCalibrationSteps() const
+{
+	TArray<FName> OutKeys;
+	CameraCalibrationStepsMap.GetKeys(OutKeys);
+	return OutKeys;
+}
+
+void UCameraCalibrationSubsystem::Initialize(FSubsystemCollectionBase& Collection)
+{
+	Super::Initialize(Collection);
+
+	PostEngineInitHandle = FCoreDelegates::OnPostEngineInit.AddLambda([&]()
+	{
+		// Find Nodal Offset Algos
+		{
+			TArray<TSubclassOf<UCameraNodalOffsetAlgo>> Algos;
+
+			for (TObjectIterator<UClass> AlgoIt; AlgoIt; ++AlgoIt)
+			{
+				if (AlgoIt->IsChildOf(UCameraNodalOffsetAlgo::StaticClass()) && !AlgoIt->HasAnyClassFlags(CLASS_Abstract | CLASS_Deprecated))
+				{
+					CameraNodalOffsetAlgosMap.Add(AlgoIt->GetFName(), TSubclassOf<UCameraNodalOffsetAlgo>(*AlgoIt));
+				}
+			}
+		}
+
+		// Find Calibration Steps
+		{
+			TArray<TSubclassOf<UCameraCalibrationStep>> Steps;
+
+			for (TObjectIterator<UClass> StepIt; StepIt; ++StepIt)
+			{
+				if (StepIt->IsChildOf(UCameraCalibrationStep::StaticClass()) && !StepIt->HasAnyClassFlags(CLASS_Abstract | CLASS_Deprecated))
+				{
+					CameraCalibrationStepsMap.Add(StepIt->GetFName(), TSubclassOf<UCameraCalibrationStep>(*StepIt));
+				}
+			}
+		}
+	});
+}
+
+void UCameraCalibrationSubsystem::Deinitialize()
+{
+	LensModelMap.Empty(0);
+	CameraNodalOffsetAlgosMap.Empty(0);
+	CameraCalibrationStepsMap.Empty(0);
+
+	if (PostEngineInitHandle.IsValid())
+	{
+		FCoreDelegates::OnPostEngineInit.Remove(PostEngineInitHandle);
+	}
+
+	Super::Deinitialize();
+}
+
+#undef LOCTEXT_NAMESPACE
+
+
diff --git a/Engine/Plugins/VirtualProduction/CameraCalibrationCore/Source/CameraCalibrationCore/Private/CameraCalibrationUtils.cpp b/Engine/Plugins/VirtualProduction/CameraCalibrationCore/Source/CameraCalibrationCore/Private/CameraCalibrationUtils.cpp
new file mode 100644
--- /dev/null
+++ b/Engine/Plugins/VirtualProduction/CameraCalibrationCore/Source/CameraCalibrationCore/Private/CameraCalibrationUtils.cpp
@@ -0,0 +1,62 @@
+// Copyright Epic Games, Inc. All Rights Reserved.
+
+
+#include "CameraCalibrationUtils.h"
+
+#include <type_traits>
+
+
+namespace CameraCalibrationUtils
+{
+	using EAxis = FCameraCalibrationUtils::EAxis;
+
+	// These axes must match the order in which they are declared in EAxis
+	static const TArray<FVector> UnitVectors =
+	{
+		{  1,  0,  0 }, //  X
+		{  0,  1,  0 }, //  Y
+		{  0,  0,  1 }, //  Z
+		{ -1,  0,  0 }, // -X
+		{  0, -1,  0 }, // -Y
+		{  0,  0, -1 }, // -Z
+	};
+
+	static const FVector& UnitVectorFromAxisEnum(EAxis Axis)
+	{
+		return UnitVectors[std::underlying_type_t<EAxis>(Axis)];
+	};
+}
+
+void FCameraCalibrationUtils::ConvertCoordinateSystem(FTransform& Transform, EAxis SrcXInDstAxis, EAxis SrcYInDstAxis, EAxis SrcZInDstAxis)
+{
+	// Unreal Engine:
+	//   Front : X
+	//   Right : Y
+	//   Up    : Z
+	//
+	// OpenCV:
+	//   Front : Z
+	//   Right : X
+	//   Up    : Yn
+
+	using namespace CameraCalibrationUtils;
+
+	FMatrix M12 = FMatrix::Identity;
+
+	M12.SetColumn(0, UnitVectorFromAxisEnum(SrcXInDstAxis));
+	M12.SetColumn(1, UnitVectorFromAxisEnum(SrcYInDstAxis));
+	M12.SetColumn(2, UnitVectorFromAxisEnum(SrcZInDstAxis));
+
+	Transform.SetFromMatrix(M12.GetTransposed() * Transform.ToMatrixWithScale() * M12);
+}
+
+void FCameraCalibrationUtils::ConvertUnrealToOpenCV(FTransform& Transform)
+{
+	ConvertCoordinateSystem(Transform, EAxis::Y, EAxis::Zn, EAxis::X);
+}
+
+void FCameraCalibrationUtils::ConvertOpenCVToUnreal(FTransform& Transform)
+{
+	ConvertCoordinateSystem(Transform, EAxis::Z, EAxis::X, EAxis::Yn);
+}
+
diff --git a/Engine/Plugins/VirtualProduction/CameraCalibrationCore/Source/CameraCalibrationCore/Private/LensDistortionModelHandlerBase.cpp b/Engine/Plugins/VirtualProduction/CameraCalibrationCore/Source/CameraCalibrationCore/Private/LensDistortionModelHandlerBase.cpp
new file mode 100644
--- /dev/null
+++ b/Engine/Plugins/VirtualProduction/CameraCalibrationCore/Source/CameraCalibrationCore/Private/LensDistortionModelHandlerBase.cpp
@@ -0,0 +1,218 @@
+// Copyright Epic Games, Inc. All Rights Reserved.
+
+#include "LensDistortionModelHandlerBase.h"
+
+#include "Algo/MaxElement.h"
+#include "Kismet/KismetRenderingLibrary.h"
+
+bool FLensDistortionState::operator==(const FLensDistortionState& Other) const
+{
+	return ((DistortionInfo.Parameters == Other.DistortionInfo.Parameters)
+		&& (DistortionInfo.FxFy == Other.DistortionInfo.FxFy)
+		&& (PrincipalPoint == Other.PrincipalPoint));
+}
+
+bool ULensDistortionModelHandlerBase::IsModelSupported(const TSubclassOf<ULensModel>& ModelToSupport) const
+{
+	return (LensModelClass == ModelToSupport);
+}
+
+void ULensDistortionModelHandlerBase::SetDistortionState(const FLensDistortionState& InNewState)
+{
+	// If the new state is equivalent to the current state, there is nothing to update. 
+	if (CurrentState != InNewState)
+	{
+		CurrentState = InNewState;
+		InterpretDistortionParameters();
+
+		bIsDirty = true;
+	}	
+}
+
+void ULensDistortionModelHandlerBase::PostInitProperties()
+{
+	Super::PostInitProperties();
+
+	// Perform handler initialization, only on derived classes
+	if (!GetClass()->HasAnyClassFlags(CLASS_Abstract))
+	{
+		InitializeHandler();
+	}
+
+	if (!HasAnyFlags(RF_ArchetypeObject | RF_ClassDefaultObject))
+	{
+		if (LensModelClass)
+		{
+			const uint32 NumDistortionParameters = LensModelClass->GetDefaultObject<ULensModel>()->GetNumParameters();
+			CurrentState.DistortionInfo.Parameters.Init(0.0f, NumDistortionParameters);
+		}
+
+		// TODO: Delay this to an initialization step where the user can specify whether to create one or both render targets
+		UndistortionDisplacementMapRT = NewObject<UTextureRenderTarget2D>(this, MakeUniqueObjectName(this, UTextureRenderTarget2D::StaticClass(), TEXT("UndistortionDisplacementMapRT")));
+		UndistortionDisplacementMapRT->RenderTargetFormat = ETextureRenderTargetFormat::RTF_RG16f;
+		UndistortionDisplacementMapRT->ClearColor = FLinearColor::Black;
+		UndistortionDisplacementMapRT->bAutoGenerateMips = false;
+		UndistortionDisplacementMapRT->InitAutoFormat(DisplacementMapWidth, DisplacementMapHeight);
+		UndistortionDisplacementMapRT->UpdateResourceImmediate(true);
+
+		DistortionDisplacementMapRT = NewObject<UTextureRenderTarget2D>(this, MakeUniqueObjectName(this, UTextureRenderTarget2D::StaticClass(), TEXT("DistortionDisplacementMapRT")));
+		DistortionDisplacementMapRT->RenderTargetFormat = ETextureRenderTargetFormat::RTF_RG16f;
+		DistortionDisplacementMapRT->ClearColor = FLinearColor::Black;
+		DistortionDisplacementMapRT->bAutoGenerateMips = false;
+		DistortionDisplacementMapRT->InitAutoFormat(DisplacementMapWidth, DisplacementMapHeight);
+		DistortionDisplacementMapRT->UpdateResourceImmediate(true);
+	}
+}
+
+#if WITH_EDITOR
+void ULensDistortionModelHandlerBase::PostEditChangeChainProperty(struct FPropertyChangedChainEvent& PropertyChangedEvent)
+{
+	const FName MemberPropertyName = PropertyChangedEvent.PropertyChain.GetActiveMemberNode()->GetValue()->GetFName();
+	if (MemberPropertyName == GET_MEMBER_NAME_CHECKED(ULensDistortionModelHandlerBase, CurrentState))
+	{
+		// Will need to revisit this init logic once we move to arbitrary lens model support 
+		if (!DistortionPostProcessMID || !UndistortionDisplacementMapMID || !DistortionDisplacementMapMID)
+		{
+			InitDistortionMaterials();
+		}
+
+		bIsDirty = true;
+		InterpretDistortionParameters();
+		SetOverscanFactor(ComputeOverscanFactor());
+		ProcessCurrentDistortion();
+	}
+}
+#endif	
+
+void ULensDistortionModelHandlerBase::SetOverscanFactor(float InOverscanFactor)
+{
+	if (!DistortionPostProcessMID || !UndistortionDisplacementMapMID || !DistortionDisplacementMapMID)
+	{
+		InitDistortionMaterials();
+	}
+	
+	OverscanFactor = InOverscanFactor;
+	if (DistortionPostProcessMID)
+	{
+		DistortionPostProcessMID->SetScalarParameterValue("overscan_factor", OverscanFactor);
+	}
+}
+
+float ULensDistortionModelHandlerBase::ComputeOverscanFactor() const
+{
+	/* Undistorted UV position in the view space:
+					   ^ View space's Y
+					   |
+			  0        1        2
+
+			  7                 3 --> View space's X
+
+			  6        5        4
+	*/
+	const TArray<FVector2D> UndistortedUVs =
+	{
+		FVector2D(0.0f, 0.0f),
+		FVector2D(0.5f, 0.0f),
+		FVector2D(1.0f, 0.0f),
+		FVector2D(1.0f, 0.5f),
+		FVector2D(1.0f, 1.0f),
+		FVector2D(0.5f, 1.0f),
+		FVector2D(0.0f, 1.0f),
+		FVector2D(0.0f, 0.5f)
+	};
+
+	TArray<float> OverscanFactors;
+	OverscanFactors.Reserve(UndistortedUVs.Num());
+	for (const FVector2D& UndistortedUV : UndistortedUVs)
+	{
+		const FVector2D DistortedUV = ComputeDistortedUV(UndistortedUV);
+		const float OverscanX = (UndistortedUV.X != 0.5f) ? (DistortedUV.X - 0.5f) / (UndistortedUV.X - 0.5f) : 1.0f;
+		const float OverscanY = (UndistortedUV.Y != 0.5f) ? (DistortedUV.Y - 0.5f) / (UndistortedUV.Y - 0.5f) : 1.0f;
+		OverscanFactors.Add(FMath::Max(OverscanX, OverscanY));
+	}
+
+	float NewOverscan = 1.0f;
+	float* MaxOverscanFactor = Algo::MaxElement(OverscanFactors);
+	if (MaxOverscanFactor != nullptr)
+	{
+		NewOverscan = FMath::Max(*MaxOverscanFactor, 1.0f);
+	}
+
+	return NewOverscan;
+}
+
+TArray<FVector2D> ULensDistortionModelHandlerBase::GetDistortedUVs(TConstArrayView<FVector2D> UndistortedUVs) const
+{
+	TArray<FVector2D> DistortedUVs;
+	DistortedUVs.Reserve(UndistortedUVs.Num());
+	for (const FVector2D& UndistortedUV : UndistortedUVs)
+	{
+		const FVector2D DistortedUV = ComputeDistortedUV(UndistortedUV);
+		DistortedUVs.Add(DistortedUV);
+	}
+	return DistortedUVs;
+}
+
+bool ULensDistortionModelHandlerBase::DrawUndistortionDisplacementMap(UTextureRenderTarget2D* DestinationTexture)
+{
+	if(DestinationTexture == nullptr)
+	{
+		return false;
+	}
+	
+	if (!DistortionPostProcessMID || !UndistortionDisplacementMapMID)
+	{
+		InitDistortionMaterials();
+	}
+	
+	UpdateMaterialParameters();
+
+	// Draw the updated displacement map render target 
+	UKismetRenderingLibrary::DrawMaterialToRenderTarget(this, DestinationTexture, UndistortionDisplacementMapMID);
+
+	return true;
+}
+
+bool ULensDistortionModelHandlerBase::DrawDistortionDisplacementMap(UTextureRenderTarget2D* DestinationTexture)
+{
+	if (DestinationTexture == nullptr)
+	{
+		return false;
+	}
+
+	if (!DistortionPostProcessMID || !DistortionDisplacementMapMID)
+	{
+		InitDistortionMaterials();
+	}
+
+	UpdateMaterialParameters();
+
+	// Draw the updated displacement map render target 
+	UKismetRenderingLibrary::DrawMaterialToRenderTarget(this, DestinationTexture, DistortionDisplacementMapMID);
+
+	return true;
+}
+
+void ULensDistortionModelHandlerBase::ProcessCurrentDistortion()
+{
+	if(bIsDirty)
+	{
+		bIsDirty = false;
+
+		InterpretDistortionParameters();
+
+		UpdateMaterialParameters();
+
+		// Draw the undistortion displacement map
+		if (UndistortionDisplacementMapMID)
+		{
+			UKismetRenderingLibrary::DrawMaterialToRenderTarget(this, UndistortionDisplacementMapRT, UndistortionDisplacementMapMID);
+		}
+
+		// Draw the distortion displacement map
+		if (DistortionDisplacementMapMID)
+		{
+			UKismetRenderingLibrary::DrawMaterialToRenderTarget(this, DistortionDisplacementMapRT, DistortionDisplacementMapMID);
+		}
+	}
+}
diff --git a/Engine/Plugins/VirtualProduction/CameraCalibrationCore/Source/CameraCalibrationCore/Private/LensFile.cpp b/Engine/Plugins/VirtualProduction/CameraCalibrationCore/Source/CameraCalibrationCore/Private/LensFile.cpp
new file mode 100644
--- /dev/null
+++ b/Engine/Plugins/VirtualProduction/CameraCalibrationCore/Source/CameraCalibrationCore/Private/LensFile.cpp
@@ -0,0 +1,783 @@
+// Copyright Epic Games, Inc. All Rights Reserved.
+
+
+#include "LensFile.h"
+
+#include "Algo/MaxElement.h"
+#include "CalibratedMapProcessor.h"
+#include "CameraCalibrationCoreLog.h"
+#include "CameraCalibrationSubsystem.h"
+#include "CineCameraComponent.h"
+#include "Engine/Engine.h"
+#include "Engine/TextureRenderTarget2D.h"
+#include "LensDistortionModelHandlerBase.h"
+#include "LensFileRendering.h"
+#include "LensInterpolationUtils.h"
+#include "Models/SphericalLensModel.h"
+
+
+namespace LensFileUtils
+{
+	UTextureRenderTarget2D* CreateDisplacementMapRenderTarget(UObject* Outer)
+	{
+		check(Outer);
+
+		//Will be good using a project settings or global resolution that user can change
+		const FIntPoint Size{256,256};
+
+		UTextureRenderTarget2D* NewRenderTarget2D = NewObject<UTextureRenderTarget2D>(Outer, MakeUniqueObjectName(Outer, UTextureRenderTarget2D::StaticClass(), TEXT("LensDisplacementMap")), RF_Public);
+		NewRenderTarget2D->RenderTargetFormat = ETextureRenderTargetFormat::RTF_RG16f;
+		NewRenderTarget2D->ClearColor = FLinearColor(0.5f, 0.5f, 0.5f, 0.5f);
+		NewRenderTarget2D->bAutoGenerateMips = false;
+		NewRenderTarget2D->bCanCreateUAV = true;
+		NewRenderTarget2D->InitAutoFormat(Size.X, Size.Y);
+		NewRenderTarget2D->UpdateResourceImmediate(true);
+
+		// Flush RHI thread after creating texture render target to make sure that RHIUpdateTextureReference is executed before doing any rendering with it
+		// This makes sure that Value->TextureReference.TextureReferenceRHI->GetReferencedTexture() is valid so that FUniformExpressionSet::FillUniformBuffer properly uses the texture for rendering, instead of using a fallback texture
+		ENQUEUE_RENDER_COMMAND(FlushRHIThreadToUpdateTextureRenderTargetReference)(
+			[](FRHICommandListImmediate& RHICmdList)
+			{
+				RHICmdList.ImmediateFlush(EImmediateFlushType::FlushRHIThread);
+			});
+
+		return NewRenderTarget2D;
+	}	
+
+
+	int32 GDistortionParametersBlendMode = 0;
+	static FAutoConsoleVariableRef CVarDistortionParametersBlendMode(
+		TEXT("Lens.DistortionParametersBlendMode"),
+		GDistortionParametersBlendMode,
+		TEXT("Integer value specifying how to blend distortion parameters: (0 - Default) Blend parameters based on blend type, (1 - Blend DisplacementMap) Run distortion equation and blend the resulting maps"),
+		ECVF_Default
+	);
+}
+
+const TArray<FVector2D> ULensFile::UndistortedUVs =
+{
+	FVector2D(0.0f, 0.0f),
+	FVector2D(0.5f, 0.0f),
+	FVector2D(1.0f, 0.0f),
+	FVector2D(1.0f, 0.5f),
+	FVector2D(1.0f, 1.0f),
+	FVector2D(0.5f, 1.0f),
+	FVector2D(0.0f, 1.0f),
+	FVector2D(0.0f, 0.5f)
+};
+
+ULensFile::ULensFile()
+{
+	LensInfo.LensModel = USphericalLensModel::StaticClass();
+	
+	if (!HasAnyFlags(RF_ArchetypeObject | RF_ClassDefaultObject))
+	{
+		CalibratedMapProcessor = MakeUnique<FCalibratedMapProcessor>();
+	}
+}
+
+#if WITH_EDITOR
+
+void ULensFile::PostEditChangeChainProperty(struct FPropertyChangedChainEvent& PropertyChangedEvent)
+{
+	if(PropertyChangedEvent.Property != nullptr)
+	{
+		const FName PropertyName = PropertyChangedEvent.Property->GetFName();
+		if (PropertyName == GET_MEMBER_NAME_CHECKED(FCalibratedMapPoint, DistortionMap))
+		{
+			//When the distortion map (stmap) changes, flag associated derived data as dirty to update it
+			check(PropertyChangedEvent.PropertyChain.GetActiveMemberNode() && PropertyChangedEvent.PropertyChain.GetActiveMemberNode()->GetValue());
+
+			const FName MemberPropertyName = PropertyChangedEvent.PropertyChain.GetActiveMemberNode()->GetValue()->GetFName();
+			const int32 ModifiedIndex = PropertyChangedEvent.GetArrayIndex(MemberPropertyName.ToString());
+			check(CalibratedMapPoints.IsValidIndex(ModifiedIndex));
+
+			FCalibratedMapPoint& MapPoint = CalibratedMapPoints[ModifiedIndex];
+			MapPoint.DerivedDistortionData.bIsDirty = true;
+		}
+		else if (PropertyName == GET_MEMBER_NAME_CHECKED(FLensInfo, LensModel))
+		{
+			uint32 NumDistortionParameters = 0;
+			if (LensInfo.LensModel)
+			{
+				NumDistortionParameters = LensInfo.LensModel.GetDefaultObject()->GetNumParameters();
+			}
+
+			for (FDistortionMapPoint& MapPoint : DistortionMapping)
+			{
+				MapPoint.DistortionInfo.Parameters.Empty();
+				MapPoint.DistortionInfo.Parameters.Init(0.0f, NumDistortionParameters);
+			}
+		}
+		else if (PropertyName == GET_MEMBER_NAME_CHECKED(ULensFile, DistortionMapping))
+		{
+			if (PropertyChangedEvent.ChangeType == EPropertyChangeType::ArrayAdd)
+			{
+				uint32 NumDistortionParameters = 0;
+				if (LensInfo.LensModel)
+				{
+					NumDistortionParameters = LensInfo.LensModel.GetDefaultObject()->GetNumParameters();
+				}
+
+				FDistortionMapPoint& MapPoint = DistortionMapping.Last();
+				MapPoint.DistortionInfo.Parameters.Init(0.0f, NumDistortionParameters);
+			}
+		}
+	}
+	
+	Super::PostEditChangeChainProperty(PropertyChangedEvent);
+}
+
+#endif //WITH_EDITOR
+
+bool ULensFile::EvaluateDistortionParameters(float InFocus, float InZoom, FDistortionInfo& OutEvaluatedValue) const
+{
+	if (DistortionMapping.Num() <= 0)
+	{
+		return false;
+	}
+
+	if (DistortionMapping.Num() == 1)
+	{
+		OutEvaluatedValue = DistortionMapping[0].DistortionInfo;
+		return true;
+	}
+
+	FDistortionMapPoint InterpPoint;
+	const bool bSuccess = LensInterpolationUtils::FIZMappingBilinearInterpolation<FDistortionMapPoint>(InFocus, InZoom, DistortionMapping, InterpPoint);
+	if (bSuccess)
+	{
+		OutEvaluatedValue = MoveTemp(InterpPoint.DistortionInfo);
+	}
+
+	return bSuccess;
+}
+
+bool ULensFile::EvaluateIntrinsicParameters(float InFocus, float InZoom, FIntrinsicParameters& OutEvaluatedValue) const
+{
+	if (IntrinsicMapping.Num() <= 0)
+	{
+		return false;
+	}
+
+	if (IntrinsicMapping.Num() == 1)
+	{
+		OutEvaluatedValue = IntrinsicMapping[0].Parameters;
+		return true;
+	}
+
+	FIntrinsicMapPoint InterpPoint;
+	const bool bSuccess = LensInterpolationUtils::FIZMappingBilinearInterpolation<FIntrinsicMapPoint>(InFocus, InZoom, IntrinsicMapping, InterpPoint);
+	if (bSuccess)
+	{
+		OutEvaluatedValue = MoveTemp(InterpPoint.Parameters);
+	}
+
+	return bSuccess;
+}
+
+bool ULensFile::EvaluateDistortionData(float InFocus, float InZoom, FVector2D InFilmback, ULensDistortionModelHandlerBase* InLensHandler, FDistortionData& OutDistortionData) const
+{	
+	if (InLensHandler == nullptr)
+	{
+		UE_LOG(LogCameraCalibrationCore, Warning, TEXT("Can't evaluate LensFile '%s' - Invalid Lens Handler"), *GetName());
+		return false;
+	}
+	
+	if (InLensHandler->GetUndistortionDisplacementMap() == nullptr)
+	{
+		UE_LOG(LogCameraCalibrationCore, Warning, TEXT("Can't evaluate LensFile '%s' - Invalid undistortion displacement map in LensHandler '%s'"), *GetName(), *InLensHandler->GetName());
+		return false;
+	}
+
+	if (InLensHandler->GetDistortionDisplacementMap() == nullptr)
+	{
+		UE_LOG(LogCameraCalibrationCore, Warning, TEXT("Can't evaluate LensFile '%s' - Invalid distortion displacement map in LensHandler '%s'"), *GetName(), *InLensHandler->GetName());
+		return false;
+	}
+
+	if (LensInfo.LensModel == nullptr)
+	{
+		UE_LOG(LogCameraCalibrationCore, Warning, TEXT("Can't evaluate LensFile '%s' - Invalid Lens Model"), *GetName());
+		SetupNoDistortionOutput(InLensHandler, OutDistortionData);
+		return false;
+	}
+
+	if (InLensHandler->IsModelSupported(LensInfo.LensModel) == false)
+	{
+		UE_LOG(LogCameraCalibrationCore, Warning, TEXT("Can't evaluate LensFile '%s' - LensHandler '%s' doesn't support lens model '%s'"), *GetName(), *InLensHandler->GetName(), *LensInfo.LensModel.GetDefaultObject()->GetModelName().ToString());
+		SetupNoDistortionOutput(InLensHandler, OutDistortionData);
+		return false;
+	}
+	
+	if(DataMode == ELensDataMode::Parameters)
+	{
+		return EvaluateDistortionForParameters(InFocus, InZoom, InFilmback, InLensHandler, OutDistortionData);
+	}
+	else
+	{
+		//Only other mode for now
+		check(DataMode == ELensDataMode::STMap);
+
+		return EvaluteDistortionForSTMaps(InFocus, InZoom, InFilmback, InLensHandler, OutDistortionData);
+	}
+}
+
+float ULensFile::ComputeOverscan(const FDistortionData& DerivedData, FVector2D PrincipalPoint) const
+{
+	//Edge case if computed data hasn't came back yet
+	if (UndistortedUVs.Num() != DerivedData.DistortedUVs.Num())
+	{
+		return 1.0f;
+	}
+
+	TArray<float, TInlineAllocator<8>> OverscanFactors;
+	OverscanFactors.Reserve(UndistortedUVs.Num());
+	for (int32 Index = 0; Index < UndistortedUVs.Num(); ++Index)
+	{
+		const FVector2D& UndistortedUV = UndistortedUVs[Index];
+		const FVector2D& DistortedUV = DerivedData.DistortedUVs[Index] + (PrincipalPoint - FVector2D(0.5f, 0.5f)) * 2.0f;
+		const float OverscanX = (UndistortedUV.X != 0.5f) ? (DistortedUV.X - 0.5f) / (UndistortedUV.X - 0.5f) : 1.0f;
+		const float OverscanY = (UndistortedUV.Y != 0.5f) ? (DistortedUV.Y - 0.5f) / (UndistortedUV.Y - 0.5f) : 1.0f;
+		OverscanFactors.Add(FMath::Max(OverscanX, OverscanY));
+	}
+
+	float* MaxOverscanFactor = Algo::MaxElement(OverscanFactors);
+	const float FoundOverscan = MaxOverscanFactor ? *MaxOverscanFactor : 1.0f;
+	
+	return FoundOverscan;
+}
+
+void ULensFile::SetupNoDistortionOutput(ULensDistortionModelHandlerBase* LensHandler, FDistortionData& OutDistortionData) const
+{
+	LensFileRendering::ClearDisplacementMap(LensHandler->GetUndistortionDisplacementMap());
+	LensFileRendering::ClearDisplacementMap(LensHandler->GetDistortionDisplacementMap());
+	OutDistortionData.DistortedUVs = UndistortedUVs;
+	OutDistortionData.OverscanFactor = 1.0f;
+	LensHandler->SetOverscanFactor(OutDistortionData.OverscanFactor);
+}
+
+bool ULensFile::EvaluateDistortionForParameters(float InFocus, float InZoom, FVector2D InFilmback, ULensDistortionModelHandlerBase* InLensHandler, FDistortionData& OutDistortionData) const
+{
+	//Parameter blending mode or no parameters to blend
+	if (LensFileUtils::GDistortionParametersBlendMode == 0
+		|| DistortionMapping.Num() <= 0)
+	{
+		FDistortionInfo DistortionPoint;
+		DistortionPoint.Parameters.SetNumZeroed(LensInfo.LensModel.GetDefaultObject()->GetNumParameters());
+		FIntrinsicParameters Intrinsic;
+		EvaluateDistortionParameters(InFocus, InZoom, DistortionPoint);
+		EvaluateIntrinsicParameters(InFocus, InZoom, Intrinsic);
+
+		//Setup handler state based on evaluated parameters. If none were found, no distortion will be returned
+		FLensDistortionState State;
+		State.DistortionInfo.Parameters = MoveTemp(DistortionPoint.Parameters);
+
+		const FVector2D FxFyScale = FVector2D(LensInfo.SensorDimensions.X / InFilmback.X, LensInfo.SensorDimensions.Y / InFilmback.Y);
+
+		State.DistortionInfo.FxFy = DistortionPoint.FxFy * FxFyScale;
+		State.PrincipalPoint = Intrinsic.PrincipalPoint;
+		
+		InLensHandler->SetDistortionState(State);
+
+		OutDistortionData.OverscanFactor = InLensHandler->ComputeOverscanFactor();
+		InLensHandler->SetOverscanFactor(OutDistortionData.OverscanFactor);
+
+		//Draw displacement map associated with the new state
+		InLensHandler->ProcessCurrentDistortion();
+	}
+	else //DisplacementMap blending
+	{
+		int32 MinMinIndex = 0;
+		int32 MinMaxIndex = 0;
+		int32 MaxMinIndex = 0;
+		int32 MaxMaxIndex = 0;
+
+		if(LensInterpolationUtils::FindBilinearInterpIndices<FDistortionMapPoint>(InFocus, InZoom, DistortionMapping, MinMinIndex, MinMaxIndex, MaxMinIndex, MaxMaxIndex))
+		{
+			if(!ensure(DistortionMapping.IsValidIndex(MinMinIndex)
+			&& DistortionMapping.IsValidIndex(MinMaxIndex)
+			&& DistortionMapping.IsValidIndex(MaxMinIndex)
+			&& DistortionMapping.IsValidIndex(MaxMaxIndex)))
+			{
+				SetupNoDistortionOutput(InLensHandler, OutDistortionData);
+				return false;	
+			}
+
+			const FDistortionMapPoint& MinMinPoint = DistortionMapping[MinMinIndex];
+			const FDistortionMapPoint& MinMaxPoint = DistortionMapping[MinMaxIndex];
+			const FDistortionMapPoint& MaxMinPoint = DistortionMapping[MaxMinIndex];
+			const FDistortionMapPoint& MaxMaxPoint = DistortionMapping[MaxMaxIndex];
+
+			FDistortionData BlendedData;
+			BlendedData.DistortedUVs.SetNumZeroed(UndistortedUVs.Num());
+
+			FDisplacementMapBlendingParams Params;
+
+			//Compute blended center shift 
+			//This will be the value used for all possible distortion points to compute the resulting displacement map
+			FIntrinsicMapPoint InterpPoint;
+			LensInterpolationUtils::FIZMappingBilinearInterpolation<FIntrinsicMapPoint>(InFocus, InZoom, IntrinsicMapping, InterpPoint);
+
+			//Common info for distortion state
+			FLensDistortionState State;
+			State.PrincipalPoint = InterpPoint.Parameters.PrincipalPoint;
+
+			//Helper function to compute the current distortion state
+			const auto GetDistortionData = [this, &State, InFilmback, InLensHandler](const FDistortionMapPoint& MapPoint, UTextureRenderTarget2D* UndistortionRenderTarget, UTextureRenderTarget2D* DistortionRenderTarget, FDistortionData& OutDistortionData)
+			{
+				State.DistortionInfo.Parameters = MapPoint.DistortionInfo.Parameters;
+
+				const FVector2D FxFyScale = FVector2D(LensInfo.SensorDimensions.X / InFilmback.X, LensInfo.SensorDimensions.Y / InFilmback.Y);
+				State.DistortionInfo.FxFy = MapPoint.DistortionInfo.FxFy * FxFyScale;
+
+				InLensHandler->SetDistortionState(State);
+				InLensHandler->DrawUndistortionDisplacementMap(UndistortionRenderTarget);
+				InLensHandler->DrawDistortionDisplacementMap(DistortionRenderTarget);
+
+				OutDistortionData.DistortedUVs = InLensHandler->GetDistortedUVs(UndistortedUVs);
+				OutDistortionData.OverscanFactor = InLensHandler->ComputeOverscanFactor();	
+			};
+
+			//Single point case
+			if (MinMinIndex == MaxMinIndex
+				&& MaxMinIndex == MinMaxIndex
+				&& MinMaxIndex == MaxMaxIndex)
+			{
+				Params.BlendType = EDisplacementMapBlendType::Passthrough;
+
+				GetDistortionData(MinMinPoint, UndistortionDisplacementMapHolders[0], DistortionDisplacementMapHolders[0], BlendedData);
+			}
+			else if (MinMinIndex == MaxMinIndex && MinMaxIndex == MaxMaxIndex)
+			{
+				//Fixed focus 
+				Params.BlendType = EDisplacementMapBlendType::Linear;
+				Params.LinearBlendFactor = LensInterpolationUtils::GetBlendFactor(InZoom, MinMinPoint.Zoom, MaxMaxPoint.Zoom);
+
+				FDistortionData DistortionData[2];
+				GetDistortionData(MinMinPoint, UndistortionDisplacementMapHolders[0], DistortionDisplacementMapHolders[0], DistortionData[0]);
+				GetDistortionData(MaxMaxPoint, UndistortionDisplacementMapHolders[1], DistortionDisplacementMapHolders[1], DistortionData[1]);
+
+				LensInterpolationUtils::Interpolate<FDistortionData>(Params.LinearBlendFactor, &DistortionData[0], &DistortionData[1], &BlendedData);
+			}
+			else if (MinMinIndex == MinMaxIndex && MaxMinIndex == MaxMaxIndex)
+			{
+				//Fixed zoom	
+				Params.BlendType = EDisplacementMapBlendType::Linear;
+				Params.LinearBlendFactor = LensInterpolationUtils::GetBlendFactor(InFocus, MinMinPoint.Focus, MaxMaxPoint.Focus);
+
+				FDistortionData DistortionData[2];
+				GetDistortionData(MinMinPoint, UndistortionDisplacementMapHolders[0], DistortionDisplacementMapHolders[0], DistortionData[0]);
+				GetDistortionData(MaxMaxPoint, UndistortionDisplacementMapHolders[1], DistortionDisplacementMapHolders[1], DistortionData[1]);
+				
+				LensInterpolationUtils::Interpolate<FDistortionData>(Params.LinearBlendFactor, &DistortionData[0], &DistortionData[1], &BlendedData);
+			}
+			else
+			{
+				//The current grid finder doesn't always yield points around the sample
+				const float X2X1 = MaxMinPoint.Focus - MinMinPoint.Focus;
+				const float Y2Y1 = MaxMaxPoint.Zoom - MinMinPoint.Zoom;
+				const float Divider = X2X1 * Y2Y1;
+
+				if (!FMath::IsNearlyZero(Divider))
+				{
+					Params.BlendType = EDisplacementMapBlendType::Bilinear;
+					Params.DeltaMinX = InFocus - MinMinPoint.Focus;
+					Params.DeltaMaxX = MaxMinPoint.Focus - InFocus;
+					Params.DeltaMinY = InZoom - MinMinPoint.Zoom;
+					Params.DeltaMaxY = MaxMaxPoint.Zoom - InZoom;
+					Params.MainCoefficient = 1.0f / Divider;
+
+					FDistortionData DistortionData[4];
+					GetDistortionData(MinMinPoint, UndistortionDisplacementMapHolders[0], DistortionDisplacementMapHolders[0], DistortionData[0]);
+					GetDistortionData(MinMaxPoint, UndistortionDisplacementMapHolders[1], DistortionDisplacementMapHolders[1], DistortionData[1]);
+					GetDistortionData(MaxMinPoint, UndistortionDisplacementMapHolders[2], DistortionDisplacementMapHolders[2], DistortionData[2]);
+					GetDistortionData(MaxMaxPoint, UndistortionDisplacementMapHolders[3], DistortionDisplacementMapHolders[3], DistortionData[3]);
+
+					LensInterpolationUtils::BilinearInterpolate<FDistortionData>(Params.MainCoefficient
+						, Params.DeltaMinX
+						, Params.DeltaMaxX
+						, Params.DeltaMinY
+						, Params.DeltaMaxY
+						, &DistortionData[0]
+						, &DistortionData[1]
+						, &DistortionData[2]
+						, &DistortionData[3]
+						, &BlendedData);
+				}
+				else
+				{
+					UE_LOG(LogCameraCalibrationCore, Warning, TEXT("Could not interpolate data for Focus = '%0.2f' and Zoom = '%0.2f' on LensFile '%s'"), InFocus, InZoom, *GetName());
+					return false;
+				}
+			}
+
+			//Draw resulting undistortion displacement map for evaluation point
+			LensFileRendering::DrawBlendedDisplacementMap(InLensHandler->GetUndistortionDisplacementMap()
+				, Params
+				, UndistortionDisplacementMapHolders[0]
+				, UndistortionDisplacementMapHolders[1]
+				, UndistortionDisplacementMapHolders[2]
+				, UndistortionDisplacementMapHolders[3]);
+
+			//Draw resulting distortion displacement map for evaluation point
+			LensFileRendering::DrawBlendedDisplacementMap(InLensHandler->GetDistortionDisplacementMap()
+				, Params
+				, DistortionDisplacementMapHolders[0]
+				, DistortionDisplacementMapHolders[1]
+				, DistortionDisplacementMapHolders[2]
+				, DistortionDisplacementMapHolders[3]);
+
+			OutDistortionData = MoveTemp(BlendedData);
+			InLensHandler->SetOverscanFactor(OutDistortionData.OverscanFactor);
+		}
+		else
+		{
+			UE_LOG(LogCameraCalibrationCore, Verbose, TEXT("Could not find distortion data for Focus = '%0.2f' and Zoom = '%0.2f' on LensFile '%s'"), InFocus, InZoom, *GetName());
+			SetupNoDistortionOutput(InLensHandler, OutDistortionData);
+			return false;
+		}
+	}
+	
+	return true;
+}
+
+bool ULensFile::EvaluteDistortionForSTMaps(float InFocus, float InZoom, FVector2D InFilmback, ULensDistortionModelHandlerBase* InLensHandler, FDistortionData& OutDistortionData) const
+{
+	if(DerivedDataInFlightCount > 0)
+	{
+		UE_LOG(LogCameraCalibrationCore, Verbose, TEXT("Can't evaluate LensFile '%s' - %d data points still being computed. Clearing render target for no distortion"), *GetName(), DerivedDataInFlightCount);
+		SetupNoDistortionOutput(InLensHandler, OutDistortionData);
+		return true;
+	}
+
+	if(CalibratedMapPoints.Num() <= 0)
+	{
+		UE_LOG(LogCameraCalibrationCore, Verbose, TEXT("Can't evaluate LensFile '%s' - No calibrated maps"), *GetName());
+		SetupNoDistortionOutput(InLensHandler, OutDistortionData);
+		return true;
+	}
+
+	int32 MinMinIndex = 0;
+	int32 MinMaxIndex = 0;
+	int32 MaxMinIndex = 0;
+	int32 MaxMaxIndex = 0;
+
+	if(LensInterpolationUtils::FindBilinearInterpIndices<FCalibratedMapPoint>(InFocus, InZoom, CalibratedMapPoints, MinMinIndex, MinMaxIndex, MaxMinIndex, MaxMaxIndex))
+	{
+		if(!ensure(CalibratedMapPoints.IsValidIndex(MinMinIndex)
+		&& CalibratedMapPoints.IsValidIndex(MinMaxIndex)
+		&& CalibratedMapPoints.IsValidIndex(MaxMinIndex)
+		&& CalibratedMapPoints.IsValidIndex(MaxMaxIndex)))
+		{
+			SetupNoDistortionOutput(InLensHandler, OutDistortionData);
+			return false;	
+		}
+
+		const FCalibratedMapPoint& MinMinPoint = CalibratedMapPoints[MinMinIndex];
+		const FCalibratedMapPoint& MinMaxPoint = CalibratedMapPoints[MinMaxIndex];
+		const FCalibratedMapPoint& MaxMinPoint = CalibratedMapPoints[MaxMinIndex];
+		const FCalibratedMapPoint& MaxMaxPoint = CalibratedMapPoints[MaxMaxIndex];
+
+		FDistortionData BlendedData;
+		BlendedData.DistortedUVs.SetNumZeroed(UndistortedUVs.Num());
+		
+		FDisplacementMapBlendingParams Params;
+		UTextureRenderTarget2D* UndistortionTextureOne = nullptr;
+		UTextureRenderTarget2D* UndistortionTextureTwo = nullptr;
+		UTextureRenderTarget2D* UndistortionTextureThree = nullptr;
+		UTextureRenderTarget2D* UndistortionTextureFour = nullptr;
+		UTextureRenderTarget2D* DistortionTextureOne = nullptr;
+		UTextureRenderTarget2D* DistortionTextureTwo = nullptr;
+		UTextureRenderTarget2D* DistortionTextureThree = nullptr;
+		UTextureRenderTarget2D* DistortionTextureFour = nullptr;
+
+		//Single point case
+		if (MinMinIndex == MaxMinIndex
+			&& MaxMinIndex == MinMaxIndex
+			&& MinMaxIndex == MaxMaxIndex)
+		{
+			Params.BlendType = EDisplacementMapBlendType::Passthrough;
+			UndistortionTextureOne = MinMinPoint.DerivedDistortionData.UndistortionDisplacementMap;
+			DistortionTextureOne = MinMinPoint.DerivedDistortionData.DistortionDisplacementMap;
+			BlendedData = MinMinPoint.DerivedDistortionData.DistortionData;
+		}
+		else if (MinMinIndex == MaxMinIndex && MinMaxIndex == MaxMaxIndex)
+		{
+			//Fixed focus 
+			Params.BlendType = EDisplacementMapBlendType::Linear;
+			Params.LinearBlendFactor = LensInterpolationUtils::GetBlendFactor(InZoom, MinMinPoint.Zoom, MaxMaxPoint.Zoom);
+			UndistortionTextureOne = MinMinPoint.DerivedDistortionData.UndistortionDisplacementMap;
+			UndistortionTextureTwo = MaxMaxPoint.DerivedDistortionData.UndistortionDisplacementMap;
+			DistortionTextureOne = MinMinPoint.DerivedDistortionData.DistortionDisplacementMap;
+			DistortionTextureTwo = MaxMaxPoint.DerivedDistortionData.DistortionDisplacementMap;
+			LensInterpolationUtils::Interpolate<FDistortionData>(Params.LinearBlendFactor, &MinMinPoint.DerivedDistortionData.DistortionData, &MaxMaxPoint.DerivedDistortionData.DistortionData, &BlendedData);
+		}
+		else if (MinMinIndex == MinMaxIndex && MaxMinIndex == MaxMaxIndex)
+		{
+			//Fixed zoom	
+			Params.BlendType = EDisplacementMapBlendType::Linear;
+			Params.LinearBlendFactor = LensInterpolationUtils::GetBlendFactor(InFocus, MinMinPoint.Focus, MaxMaxPoint.Focus);
+			UndistortionTextureOne = MinMinPoint.DerivedDistortionData.UndistortionDisplacementMap;
+			UndistortionTextureTwo = MaxMaxPoint.DerivedDistortionData.UndistortionDisplacementMap;
+			DistortionTextureOne = MinMinPoint.DerivedDistortionData.DistortionDisplacementMap;
+			DistortionTextureTwo = MaxMaxPoint.DerivedDistortionData.DistortionDisplacementMap;
+			LensInterpolationUtils::Interpolate<FDistortionData>(Params.LinearBlendFactor, &MinMinPoint.DerivedDistortionData.DistortionData, &MaxMaxPoint.DerivedDistortionData.DistortionData, &BlendedData);
+		}
+		else
+		{
+			//The current grid finder doesn't always yield points around the sample
+			const float X2X1 = MaxMinPoint.Focus - MinMinPoint.Focus;
+			const float Y2Y1 = MaxMaxPoint.Zoom - MinMinPoint.Zoom;
+			const float Divider = X2X1 * Y2Y1;
+
+			if (!FMath::IsNearlyZero(Divider))
+			{
+				Params.BlendType = EDisplacementMapBlendType::Bilinear;
+				Params.DeltaMinX = InFocus - MinMinPoint.Focus;
+				Params.DeltaMaxX = MaxMinPoint.Focus - InFocus;
+				Params.DeltaMinY = InZoom - MinMinPoint.Zoom;
+				Params.DeltaMaxY = MaxMaxPoint.Zoom - InZoom;
+				Params.MainCoefficient = 1.0f / Divider;
+				UndistortionTextureOne = MinMinPoint.DerivedDistortionData.UndistortionDisplacementMap;
+				UndistortionTextureTwo = MinMaxPoint.DerivedDistortionData.UndistortionDisplacementMap;
+				UndistortionTextureThree = MaxMinPoint.DerivedDistortionData.UndistortionDisplacementMap;
+				UndistortionTextureFour = MaxMaxPoint.DerivedDistortionData.UndistortionDisplacementMap;
+				DistortionTextureOne = MinMinPoint.DerivedDistortionData.DistortionDisplacementMap;
+				DistortionTextureTwo = MinMaxPoint.DerivedDistortionData.DistortionDisplacementMap;
+				DistortionTextureThree = MaxMinPoint.DerivedDistortionData.DistortionDisplacementMap;
+				DistortionTextureFour = MaxMaxPoint.DerivedDistortionData.DistortionDisplacementMap;
+				LensInterpolationUtils::BilinearInterpolate<FDistortionData>(Params.MainCoefficient
+				, Params.DeltaMinX
+				, Params.DeltaMaxX
+				, Params.DeltaMinY
+				, Params.DeltaMaxY
+				, &MinMinPoint.DerivedDistortionData.DistortionData
+				, &MinMaxPoint.DerivedDistortionData.DistortionData
+				, &MaxMaxPoint.DerivedDistortionData.DistortionData
+				, &MaxMaxPoint.DerivedDistortionData.DistortionData
+				, &BlendedData);
+			}
+			else
+			{
+				UE_LOG(LogCameraCalibrationCore, Warning, TEXT("Could not interpolate data for Focus = '%0.2f' and Zoom = '%0.2f' on LensFile '%s'"), InFocus, InZoom, *GetName());
+				return false;
+			}
+		}
+
+		//Compute blended principal point to apply on top of blended displacement map
+		FIntrinsicMapPoint InterpPoint;
+		LensInterpolationUtils::FIZMappingBilinearInterpolation<FIntrinsicMapPoint>(InFocus, InZoom, IntrinsicMapping, InterpPoint);
+
+		// Compute a scale factor to scale the UVs in the blended displacement map
+		Params.FxFyScale = FVector2D(InFilmback.X / LensInfo.SensorDimensions.X, InFilmback.Y / LensInfo.SensorDimensions.Y);
+		Params.PrincipalPoint = InterpPoint.Parameters.PrincipalPoint;
+
+		//Draw resulting undistortion displacement map for evaluation point
+		LensFileRendering::DrawBlendedDisplacementMap(InLensHandler->GetUndistortionDisplacementMap()
+			, Params
+			, UndistortionTextureOne
+			, UndistortionTextureTwo
+			, UndistortionTextureThree
+			, UndistortionTextureFour);
+
+		//Draw resulting displacement map for evaluation point
+		if(LensFileRendering::DrawBlendedDisplacementMap(InLensHandler->GetDistortionDisplacementMap()
+			, Params
+			, DistortionTextureOne
+			, DistortionTextureTwo
+			, DistortionTextureThree
+			, DistortionTextureFour))
+		{
+			OutDistortionData.OverscanFactor = ComputeOverscan(BlendedData, Params.PrincipalPoint);
+			InLensHandler->SetOverscanFactor(OutDistortionData.OverscanFactor);
+		}
+	}
+	else
+	{
+		UE_LOG(LogCameraCalibrationCore, Verbose, TEXT("Could not find distortion data for Focus = '%0.2f' and Zoom = '%0.2f' on LensFile '%s'"), InFocus, InZoom, *GetName());
+		SetupNoDistortionOutput(InLensHandler, OutDistortionData);
+		return false;
+	}
+
+	return true;
+}
+
+bool ULensFile::EvaluateNodalPointOffset(float InFocus, float InZoom, FNodalPointOffset& OutEvaluatedValue) const
+{
+	if (NodalOffsetMapping.Num() <= 0)
+	{
+		return false;
+	}
+
+	if (NodalOffsetMapping.Num() == 1)
+	{
+		OutEvaluatedValue = NodalOffsetMapping[0].NodalOffset;
+		return true;
+	}
+
+	FNodalOffsetMapPoint InterpPoint;
+	const bool bSuccess = LensInterpolationUtils::FIZMappingBilinearInterpolation<FNodalOffsetMapPoint>(InFocus, InZoom, NodalOffsetMapping, InterpPoint);
+	if (bSuccess)
+	{
+		OutEvaluatedValue = MoveTemp(InterpPoint.NodalOffset);
+	}
+
+	return bSuccess;
+}
+
+bool ULensFile::HasFocusEncoderMapping() const
+{
+	return EncoderMapping.Focus.GetNumKeys() > 0;
+}
+
+float ULensFile::EvaluateNormalizedFocus(float InNormalizedValue)
+{
+	return EncoderMapping.Focus.Eval(InNormalizedValue);
+}
+
+bool ULensFile::HasIrisEncoderMapping() const
+{
+	return EncoderMapping.Iris.GetNumKeys() > 0;
+}
+
+float ULensFile::EvaluateNormalizedIris(float InNormalizedValue)
+{
+	return EncoderMapping.Iris.Eval(InNormalizedValue);
+}
+
+bool ULensFile::HasZoomEncoderMapping() const
+{
+	return EncoderMapping.Zoom.GetNumKeys() > 0;
+}
+
+float ULensFile::EvaluateNormalizedZoom(float InNormalizedValue)
+{
+	return EncoderMapping.Zoom.Eval(InNormalizedValue);
+}
+
+void ULensFile::OnDistortionDerivedDataJobCompleted(const FDerivedDistortionDataJobOutput& JobOutput)
+{
+	//Keep track of jobs being processed
+	--DerivedDataInFlightCount;
+	
+	const FGuid PointIdentifier = JobOutput.Identifier;
+	if(FCalibratedMapPoint* AssociatedPoint = CalibratedMapPoints.FindByPredicate([PointIdentifier](const FCalibratedMapPoint& Other){ return Other.GetIdentifier() == PointIdentifier; }))
+	{
+		if (JobOutput.Result == EDerivedDistortionDataResult::Success)
+		{
+			AssociatedPoint->DerivedDistortionData.DistortionData.DistortedUVs = JobOutput.EdgePointsDistortedUVs;
+		}
+		else
+		{
+			UE_LOG(LogCameraCalibrationCore, Warning, TEXT("Could not derive distortion data for calibrated map point with Focus = '%0.2f' and Zoom = '%0.2f' on LensFile '%s'"), AssociatedPoint->Focus, AssociatedPoint->Zoom, *GetName());
+		}
+	}
+}
+
+bool ULensFile::IsCineCameraCompatible(const UCineCameraComponent* CineCameraComponent)
+{
+	if (!CineCameraComponent)
+	{ 
+		return false;
+	}
+
+	if (DataMode == ELensDataMode::STMap)
+	{
+		if ((LensInfo.SensorDimensions.X < CineCameraComponent->Filmback.SensorWidth)
+			|| (LensInfo.SensorDimensions.Y < CineCameraComponent->Filmback.SensorHeight))
+		{
+			return false;
+		}
+	}
+
+	return true;
+}
+
+
+void ULensFile::PostInitProperties()
+{
+	Super::PostInitProperties();
+	
+	//Create displacement maps used when blending them together to get final distortion map
+	UndistortionDisplacementMapHolders.Reserve(DisplacementMapHolderCount);
+	DistortionDisplacementMapHolders.Reserve(DisplacementMapHolderCount);
+	for (int32 Index = 0; Index < DisplacementMapHolderCount; ++Index)
+	{
+		UTextureRenderTarget2D* NewMap = LensFileUtils::CreateDisplacementMapRenderTarget(GetTransientPackage());
+		UndistortionDisplacementMapHolders.Add(NewMap);
+		DistortionDisplacementMapHolders.Add(NewMap);
+	}
+}
+
+void ULensFile::Tick(float DeltaTime)
+{
+	if (CalibratedMapProcessor)
+	{
+		CalibratedMapProcessor->Update();
+	}
+
+	UpdateDerivedData();
+}
+
+TStatId ULensFile::GetStatId() const
+{
+	RETURN_QUICK_DECLARE_CYCLE_STAT(ULensFile, STATGROUP_Tickables);
+}
+
+void ULensFile::UpdateDerivedData()
+{
+	if(DataMode == ELensDataMode::STMap)
+	{
+		for(FCalibratedMapPoint& MapPoint : CalibratedMapPoints)
+		{
+			if (MapPoint.DerivedDistortionData.bIsDirty)
+			{
+				//Create required undistortion texture for newly added points
+				if(MapPoint.DerivedDistortionData.UndistortionDisplacementMap == nullptr)
+				{
+					MapPoint.DerivedDistortionData.UndistortionDisplacementMap = LensFileUtils::CreateDisplacementMapRenderTarget(this);
+				}
+
+				//Create required distortion texture for newly added points
+				if (MapPoint.DerivedDistortionData.DistortionDisplacementMap == nullptr)
+				{
+					MapPoint.DerivedDistortionData.DistortionDisplacementMap = LensFileUtils::CreateDisplacementMapRenderTarget(this);
+				}
+
+				check(MapPoint.DerivedDistortionData.UndistortionDisplacementMap);
+				check(MapPoint.DerivedDistortionData.DistortionDisplacementMap);
+
+				FDerivedDistortionDataJobArgs JobArgs;
+				JobArgs.Identifier = MapPoint.GetIdentifier();
+				JobArgs.SourceDistortionMap = MapPoint.DistortionMap;
+				JobArgs.OutputUndistortionDisplacementMap = MapPoint.DerivedDistortionData.UndistortionDisplacementMap;
+				JobArgs.OutputDistortionDisplacementMap = MapPoint.DerivedDistortionData.DistortionDisplacementMap;
+				JobArgs.JobCompletedCallback.BindUObject(this, &ULensFile::OnDistortionDerivedDataJobCompleted);
+				if (CalibratedMapProcessor->PushDerivedDistortionDataJob(MoveTemp(JobArgs)))
+				{
+					++DerivedDataInFlightCount;
+					MapPoint.DerivedDistortionData.bIsDirty = false;
+				}
+			}
+		}
+	}
+}
+
+ULensFile* FLensFilePicker::GetLensFile() const
+{
+	ULensFile* ReturnedLens = nullptr;
+
+	if (bOverrideDefaultLensFile)
+	{
+		ReturnedLens = LensFile;
+	}
+	else if(GEngine)
+	{
+		UCameraCalibrationSubsystem* SubSystem = GEngine->GetEngineSubsystem<UCameraCalibrationSubsystem>();
+		ReturnedLens = SubSystem->GetDefaultLensFile();
+	}
+
+	return ReturnedLens;
+}
+
diff --git a/Engine/Plugins/VirtualProduction/CameraCalibrationCore/Source/CameraCalibrationCore/Private/LensFileRendering.cpp b/Engine/Plugins/VirtualProduction/CameraCalibrationCore/Source/CameraCalibrationCore/Private/LensFileRendering.cpp
new file mode 100644
--- /dev/null
+++ b/Engine/Plugins/VirtualProduction/CameraCalibrationCore/Source/CameraCalibrationCore/Private/LensFileRendering.cpp
@@ -0,0 +1,170 @@
+// Copyright Epic Games, Inc. All Rights Reserved.
+
+
+#include "LensFileRendering.h"
+
+#include "Engine/TextureRenderTarget2D.h"
+#include "GlobalShader.h"											   
+#include "PixelShaderUtils.h"
+#include "RenderGraphUtils.h"
+#include "ShaderParameterStruct.h"
+#include "ScreenPass.h"
+#include "TextureResource.h"
+
+
+/** A pixel shader that blends displacement maps together. */
+class FDisplacementMapBlendPS : public FGlobalShader
+{
+	DECLARE_GLOBAL_SHADER(FDisplacementMapBlendPS);
+	SHADER_USE_PARAMETER_STRUCT(FDisplacementMapBlendPS, FGlobalShader);
+
+	class FBlendType : SHADER_PERMUTATION_INT("BLEND_TYPE", 3);
+	using FPermutationDomain = TShaderPermutationDomain<FBlendType>;
+	
+	BEGIN_SHADER_PARAMETER_STRUCT(FParameters, )
+
+		SHADER_PARAMETER(float, BlendFactor)
+		SHADER_PARAMETER(float, MainCoefficient)
+		SHADER_PARAMETER(float, DeltaMinX)
+		SHADER_PARAMETER(float, DeltaMaxX)
+		SHADER_PARAMETER(float, DeltaMinY)
+		SHADER_PARAMETER(float, DeltaMaxY)
+		SHADER_PARAMETER(FVector2D, FxFyScale)
+		SHADER_PARAMETER(FVector2D, PrincipalPoint)
+		SHADER_PARAMETER(FIntPoint, OutputTextureExtent)
+        SHADER_PARAMETER_RDG_TEXTURE(Texture2D, SourceTextureOne)
+        SHADER_PARAMETER_RDG_TEXTURE(Texture2D, SourceTextureTwo)
+        SHADER_PARAMETER_RDG_TEXTURE(Texture2D, SourceTextureThree)
+        SHADER_PARAMETER_RDG_TEXTURE(Texture2D, SourceTextureFour)
+        SHADER_PARAMETER_SAMPLER(SamplerState, SourceTextureSampler)
+
+        RENDER_TARGET_BINDING_SLOTS()
+    END_SHADER_PARAMETER_STRUCT()
+
+	// Called by the engine to determine which permutations to compile for this shader
+    static bool ShouldCompilePermutation(const FGlobalShaderPermutationParameters& Parameters)
+	{
+		return IsFeatureLevelSupported(Parameters.Platform, ERHIFeatureLevel::SM5);
+	}
+};
+
+IMPLEMENT_GLOBAL_SHADER(FDisplacementMapBlendPS, "/Plugin/CameraCalibrationCore/Private/DisplacementMapBlending.usf", "BlendPS", SF_Pixel);
+
+
+namespace LensFileRendering
+{
+	void ClearDisplacementMap(UTextureRenderTarget2D* OutRenderTarget)
+	{
+		if (OutRenderTarget != nullptr)
+		{
+			const FTextureRenderTargetResource* const DestinationTextureResource = OutRenderTarget->GameThread_GetRenderTargetResource();
+
+			ENQUEUE_RENDER_COMMAND(LensFileRendering_ClearDisplacementMap)(
+			[DestinationTextureResource](FRHICommandListImmediate& RHICmdList)
+			{
+				FRDGBuilder GraphBuilder(RHICmdList);
+
+				const FRDGTextureRef OutputTexture = GraphBuilder.RegisterExternalTexture(CreateRenderTarget(DestinationTextureResource->TextureRHI, TEXT("OutputDisplacement")));
+
+				FLinearColor NoDistortionColor(0.0f, 0.0f, 0.0f, 0.0f);
+				AddClearRenderTargetPass(GraphBuilder, OutputTexture, NoDistortionColor);
+
+				GraphBuilder.Execute();
+			});
+		}
+	}	
+	
+	bool DrawBlendedDisplacementMap(UTextureRenderTarget2D* OutRenderTarget
+	, const FDisplacementMapBlendingParams& BlendParams
+	, UTextureRenderTarget2D* SourceTextureOne
+	, UTextureRenderTarget2D* SourceTextureTwo
+	, UTextureRenderTarget2D* SourceTextureThree
+	, UTextureRenderTarget2D* SourceTextureFour)
+{
+	if (SourceTextureOne == nullptr || OutRenderTarget == nullptr)
+	{
+		return false;
+	}
+
+	const FTextureRenderTargetResource* const SourceTextureOneResource = SourceTextureOne->GameThread_GetRenderTargetResource();
+	const FTextureRenderTargetResource* const SourceTextureTwoResource = SourceTextureTwo ? SourceTextureTwo->GameThread_GetRenderTargetResource() : nullptr;
+	const FTextureRenderTargetResource* const SourceTextureThreeResource = SourceTextureThree? SourceTextureThree->GameThread_GetRenderTargetResource() : nullptr;
+	const FTextureRenderTargetResource* const SourceTextureFourResource = SourceTextureFour ? SourceTextureFour->GameThread_GetRenderTargetResource() : nullptr;
+	const FTextureRenderTargetResource* const DestinationTextureResource = OutRenderTarget->GameThread_GetRenderTargetResource();
+
+	ENQUEUE_RENDER_COMMAND(LensFileRendering_DrawBlendedDisplacementMap)(
+		[SourceTextureOneResource, SourceTextureTwoResource, SourceTextureThreeResource, SourceTextureFourResource, BlendParams = BlendParams, DestinationTextureResource](FRHICommandListImmediate& RHICmdList)
+		{
+			FRDGBuilder GraphBuilder(RHICmdList);
+
+			FDisplacementMapBlendPS::FParameters* PassParameters = GraphBuilder.AllocParameters<FDisplacementMapBlendPS::FParameters>();
+
+			//Setup always used parameters
+			FRDGTextureRef TextureOne = GraphBuilder.RegisterExternalTexture(CreateRenderTarget(SourceTextureOneResource->TextureRHI, TEXT("DisplacementMapOne")));
+			FRDGTextureRef OutputTexture = GraphBuilder.RegisterExternalTexture(CreateRenderTarget(DestinationTextureResource->TextureRHI, TEXT("OutputDisplacement")));
+			PassParameters->OutputTextureExtent = OutputTexture->Desc.Extent;
+			PassParameters->SourceTextureSampler = TStaticSamplerState<SF_Bilinear, AM_Clamp, AM_Clamp, AM_Clamp>::GetRHI(); 
+			PassParameters->SourceTextureOne = TextureOne;
+			PassParameters->FxFyScale = BlendParams.FxFyScale;
+			PassParameters->PrincipalPoint = BlendParams.PrincipalPoint;
+
+			//Setup parameters based on blending type
+			FDisplacementMapBlendPS::FPermutationDomain PermutationVector;
+			switch(BlendParams.BlendType)
+			{
+			case EDisplacementMapBlendType::Linear:
+				{
+					PermutationVector.Set<FDisplacementMapBlendPS::FBlendType>(1);
+					check(SourceTextureTwoResource);
+					FRDGTextureRef TextureTwo = GraphBuilder.RegisterExternalTexture(CreateRenderTarget(SourceTextureTwoResource->TextureRHI, TEXT("DisplacementMapTwo")));
+					PassParameters->BlendFactor = BlendParams.LinearBlendFactor;
+					PassParameters->SourceTextureTwo = TextureTwo;
+					break;
+				}
+			case EDisplacementMapBlendType::Bilinear:
+				{
+					PermutationVector.Set<FDisplacementMapBlendPS::FBlendType>(2);
+					check(SourceTextureTwoResource && SourceTextureThreeResource && SourceTextureFourResource);
+					FRDGTextureRef TextureTwo = GraphBuilder.RegisterExternalTexture(CreateRenderTarget(SourceTextureTwoResource->TextureRHI, TEXT("DisplacementMapTwo")));
+					FRDGTextureRef TextureThree = GraphBuilder.RegisterExternalTexture(CreateRenderTarget(SourceTextureThreeResource->TextureRHI, TEXT("DisplacementMapThree")));
+					FRDGTextureRef TextureFour = GraphBuilder.RegisterExternalTexture(CreateRenderTarget(SourceTextureFourResource->TextureRHI, TEXT("DisplacementMapFour")));
+					PassParameters->MainCoefficient = BlendParams.MainCoefficient;
+					PassParameters->DeltaMinX = BlendParams.DeltaMinX;
+					PassParameters->DeltaMaxX = BlendParams.DeltaMaxX;
+					PassParameters->DeltaMinY = BlendParams.DeltaMinY;
+					PassParameters->DeltaMaxY = BlendParams.DeltaMaxY;
+					PassParameters->SourceTextureTwo = TextureTwo;
+					PassParameters->SourceTextureThree = TextureThree;
+					PassParameters->SourceTextureFour = TextureFour;
+						
+					break;
+				}
+			case EDisplacementMapBlendType::Passthrough:
+			default:
+				{
+					PermutationVector.Set<FDisplacementMapBlendPS::FBlendType>(0);
+					break;
+				}
+			}
+
+			FGlobalShaderMap* GlobalShaderMap = GetGlobalShaderMap(GMaxRHIFeatureLevel);
+			const TShaderMapRef<FDisplacementMapBlendPS> PixelShader(GlobalShaderMap, PermutationVector);
+
+			FScreenPassRenderTarget SceneColorRenderTarget(OutputTexture, ERenderTargetLoadAction::ELoad);
+			PassParameters->RenderTargets[0] = SceneColorRenderTarget.GetRenderTargetBinding();
+
+			FPixelShaderUtils::AddFullscreenPass(
+						GraphBuilder
+						, GlobalShaderMap
+						, RDG_EVENT_NAME("BlendingLensDisplacementMap")
+						, PixelShader
+						, PassParameters
+						, FIntRect(0, 0, OutputTexture->Desc.Extent.X, OutputTexture->Desc.Extent.Y));            	
+
+			GraphBuilder.Execute();
+		});
+	
+	return true;
+}
+}
+
diff --git a/Engine/Plugins/VirtualProduction/CameraCalibrationCore/Source/CameraCalibrationCore/Private/LensFileRendering.h b/Engine/Plugins/VirtualProduction/CameraCalibrationCore/Source/CameraCalibrationCore/Private/LensFileRendering.h
new file mode 100644
--- /dev/null
+++ b/Engine/Plugins/VirtualProduction/CameraCalibrationCore/Source/CameraCalibrationCore/Private/LensFileRendering.h
@@ -0,0 +1,56 @@
+// Copyright Epic Games, Inc. All Rights Reserved.
+
+#pragma once
+
+#include "CoreTypes.h"
+#include "Math/Vector2D.h"
+
+class UTextureRenderTarget2D;
+
+/** Types of blending used which drives permutation */
+enum class EDisplacementMapBlendType :uint8
+{
+	Passthrough,
+	Linear,
+	Bilinear,
+};
+
+/** Single struct containing blending params for all types */
+struct FDisplacementMapBlendingParams
+{
+	/** Active type of blending */
+	EDisplacementMapBlendType BlendType = EDisplacementMapBlendType::Passthrough;
+
+	/** Linear blend parameter */
+	float LinearBlendFactor = 0.0f;
+
+	/** Bilinear blend parameters */
+	float MainCoefficient = 0.0f;
+	float DeltaMinX = 0.0f;
+	float DeltaMaxX = 0.0f;
+	float DeltaMinY = 0.0f;
+	float DeltaMaxY = 0.0f;
+
+	/** Scale parameter that allows displacement maps for one sensor size to be applied to camera's with a different sensor size */
+	FVector2D FxFyScale = { 1.0f, 1.0f };
+
+	/** Image center parameter to compute center shift needed to offset resulting map */
+	FVector2D PrincipalPoint = { 0.5f, 0.5f };
+};
+
+namespace LensFileRendering
+{
+	/** Clears the given render target. Useful when no distortion can be applied and the RT has to be resetted */
+	void ClearDisplacementMap(UTextureRenderTarget2D* OutRenderTarget);
+
+	/**
+	 * Draws the blended result of displacement map from input textures based on blend parameters
+	 * One texture is always needed to do a passthrough. Up to four textures can be blended using bilinear
+	 */
+	bool DrawBlendedDisplacementMap(UTextureRenderTarget2D* OutRenderTarget
+		, const FDisplacementMapBlendingParams& BlendParams
+		, UTextureRenderTarget2D* SourceTextureOne
+		, UTextureRenderTarget2D* SourceTextureTwo = nullptr
+		, UTextureRenderTarget2D* SourceTextureThree = nullptr
+		, UTextureRenderTarget2D* SourceTextureFour = nullptr);
+}
diff --git a/Engine/Plugins/VirtualProduction/CameraCalibrationCore/Source/CameraCalibrationCore/Private/LensInterpolationUtils.cpp b/Engine/Plugins/VirtualProduction/CameraCalibrationCore/Source/CameraCalibrationCore/Private/LensInterpolationUtils.cpp
new file mode 100644
--- /dev/null
+++ b/Engine/Plugins/VirtualProduction/CameraCalibrationCore/Source/CameraCalibrationCore/Private/LensInterpolationUtils.cpp
@@ -0,0 +1,323 @@
+// Copyright Epic Games, Inc. All Rights Reserved.
+
+
+#include "LensInterpolationUtils.h"
+
+#include "LensData.h"
+
+
+//Property interpolation utils largely inspired from livelink interp code
+namespace LensInterpolationUtils
+{
+	void BilinearInterpolateProperty(FProperty* Property
+		, float MainCoefficient
+		, float DeltaMinFocus
+		, float DeltaMaxFocus
+		, float DeltaMinZoom
+		, float DeltaMaxZoom
+		, const void* DataA, const void* InDataB, const void* InDataC, const void* InDataD, void* OutData);
+
+	void InterpolateProperty(FProperty* Property, float InBlendWeight, const void* InFrameDataA, const void* InFrameDataB, void* OutFrameData);
+
+	template<typename Type>
+	void BilinearInterpolate(const FStructProperty* StructProperty
+		, float MainCoefficient
+		, float DeltaMinFocus
+		, float DeltaMaxFocus
+		, float DeltaMinZoom
+		, float DeltaMaxZoom
+		, const void* DataA
+		, const void* DataB
+		, const void* DataC
+		, const void* DataD
+		, void* DataResult)
+	{
+		const Type* ValuePtrA = StructProperty->ContainerPtrToValuePtr<Type>(DataA);
+		const Type* ValuePtrB = StructProperty->ContainerPtrToValuePtr<Type>(DataB);
+		const Type* ValuePtrC = StructProperty->ContainerPtrToValuePtr<Type>(DataC);
+		const Type* ValuePtrD = StructProperty->ContainerPtrToValuePtr<Type>(DataD);
+		Type* ValuePtrResult = StructProperty->ContainerPtrToValuePtr<Type>(DataResult);
+
+		Type ValueResult = BilinearBlendValue(MainCoefficient, DeltaMinFocus, DeltaMaxFocus, DeltaMinZoom, DeltaMaxZoom, *ValuePtrA, *ValuePtrB, *ValuePtrC, *ValuePtrD);
+		StructProperty->CopySingleValue(ValuePtrResult, &ValueResult);
+	}
+
+	void BilinearInterpolate(const UStruct* InStruct
+		, float MainCoefficient
+		, float DeltaMinFocus
+		, float DeltaMaxFocus
+		, float DeltaMinZoom
+		, float DeltaMaxZoom
+		, const void* DataA, const void* DataB, const void* DataC, const void* DataD, void* OutFrameData)
+	{
+		for (TFieldIterator<FProperty> Itt(InStruct); Itt; ++Itt)
+		{
+			FProperty* Property = *Itt;
+
+			//Add support for arrays if required
+			if (FArrayProperty* ArrayProperty = CastField<FArrayProperty>(Property))
+			{
+				//ArrayProps have an ArrayDim of 1 but just to be sure...
+				for (int32 DimIndex = 0; DimIndex < ArrayProperty->ArrayDim; ++DimIndex)
+				{
+					const void* FrameDataA = Property->ContainerPtrToValuePtr<const void>(DataA, DimIndex);
+					const void* FrameDataB = Property->ContainerPtrToValuePtr<const void>(DataB, DimIndex);
+					const void* FrameDataC = Property->ContainerPtrToValuePtr<const void>(DataC, DimIndex);
+					const void* FrameDataD = Property->ContainerPtrToValuePtr<const void>(DataD, DimIndex);
+					void* DataResult = Property->ContainerPtrToValuePtr<void>(OutFrameData, DimIndex);
+
+					FScriptArrayHelper ArrayHelperA(ArrayProperty, FrameDataA);
+					FScriptArrayHelper ArrayHelperB(ArrayProperty, FrameDataB);
+					FScriptArrayHelper ArrayHelperC(ArrayProperty, FrameDataC);
+					FScriptArrayHelper ArrayHelperD(ArrayProperty, FrameDataD);
+					FScriptArrayHelper ArrayHelperResult(ArrayProperty, DataResult);
+
+					const int32 MinValue = FMath::Min(ArrayHelperA.Num(), FMath::Min(ArrayHelperB.Num(), FMath::Min(ArrayHelperC.Num(), ArrayHelperD.Num())));
+					ArrayHelperResult.Resize(MinValue);
+
+					for (int32 ArrayIndex = 0; ArrayIndex < MinValue; ++ArrayIndex)
+					{
+						BilinearInterpolateProperty(ArrayProperty->Inner, MainCoefficient, DeltaMinFocus, DeltaMaxFocus, DeltaMinZoom, DeltaMaxZoom
+							, ArrayHelperA.GetRawPtr(ArrayIndex)
+							, ArrayHelperB.GetRawPtr(ArrayIndex)
+							, ArrayHelperC.GetRawPtr(ArrayIndex)
+							, ArrayHelperD.GetRawPtr(ArrayIndex)
+							, ArrayHelperResult.GetRawPtr(ArrayIndex));
+					}
+				}
+			}
+			else if (Property->ArrayDim > 1)
+			{
+				for (int32 DimIndex = 0; DimIndex < Property->ArrayDim; ++DimIndex)
+				{
+					const void* FrameDataA = Property->ContainerPtrToValuePtr<const void>(DataA, DimIndex);
+					const void* FrameDataB = Property->ContainerPtrToValuePtr<const void>(DataB, DimIndex);
+					const void* FrameDataC = Property->ContainerPtrToValuePtr<const void>(DataC, DimIndex);
+					const void* FrameDataD = Property->ContainerPtrToValuePtr<const void>(DataD, DimIndex);
+					void* DataResult = Property->ContainerPtrToValuePtr<void>(OutFrameData, DimIndex);
+
+					BilinearInterpolateProperty(Property, MainCoefficient, DeltaMinFocus, DeltaMaxFocus, DeltaMinZoom, DeltaMaxZoom, FrameDataA, FrameDataB, FrameDataC, FrameDataD, DataResult);
+				}
+			}
+			else
+			{
+				BilinearInterpolateProperty(Property, MainCoefficient, DeltaMinFocus, DeltaMaxFocus, DeltaMinZoom, DeltaMaxZoom, DataA, DataB, DataC, DataD, OutFrameData);
+			}
+
+		}
+	}
+
+	void BilinearInterpolateProperty(FProperty* Property
+		, float MainCoefficient
+		, float DeltaMinFocus
+		, float DeltaMaxFocus
+		, float DeltaMinZoom
+		, float DeltaMaxZoom
+		, const void* InFrameDataA, const void* InFrameDataB, const void* InFrameDataC, const void* InFrameDataD, void* OutFrameData)
+	{
+		if (FStructProperty* StructProperty = CastField<FStructProperty>(Property))
+		{
+			if (StructProperty->Struct->GetFName() == NAME_Vector)
+			{
+				BilinearInterpolate<FVector>(StructProperty, MainCoefficient, DeltaMinFocus, DeltaMaxFocus, DeltaMinZoom, DeltaMaxZoom, InFrameDataA, InFrameDataB, InFrameDataC, InFrameDataD, OutFrameData);
+			}
+			else if (StructProperty->Struct->GetFName() == NAME_Vector2D)
+			{
+				BilinearInterpolate<FVector2D>(StructProperty, MainCoefficient, DeltaMinFocus, DeltaMaxFocus, DeltaMinZoom, DeltaMaxZoom, InFrameDataA, InFrameDataB, InFrameDataC, InFrameDataD, OutFrameData);
+			}
+			else if (StructProperty->Struct->GetFName() == NAME_Vector4)
+			{
+				BilinearInterpolate<FVector4>(StructProperty, MainCoefficient, DeltaMinFocus, DeltaMaxFocus, DeltaMinZoom, DeltaMaxZoom, InFrameDataA, InFrameDataB, InFrameDataC, InFrameDataD, OutFrameData);
+			}
+			else if (StructProperty->Struct->GetFName() == NAME_Rotator)
+			{
+				BilinearInterpolate<FRotator>(StructProperty, MainCoefficient, DeltaMinFocus, DeltaMaxFocus, DeltaMinZoom, DeltaMaxZoom, InFrameDataA, InFrameDataB, InFrameDataC, InFrameDataD, OutFrameData);
+			}
+			else if (StructProperty->Struct->GetFName() == NAME_Quat)
+			{
+				BilinearInterpolate<FQuat>(StructProperty, MainCoefficient, DeltaMinFocus, DeltaMaxFocus, DeltaMinZoom, DeltaMaxZoom, InFrameDataA, InFrameDataB, InFrameDataC, InFrameDataD, OutFrameData);
+			}
+			else
+			{
+				const void* Data0 = StructProperty->ContainerPtrToValuePtr<const void>(InFrameDataA);
+				const void* Data1 = StructProperty->ContainerPtrToValuePtr<const void>(InFrameDataB);
+				const void* Data2 = StructProperty->ContainerPtrToValuePtr<const void>(InFrameDataC);
+				const void* Data3 = StructProperty->ContainerPtrToValuePtr<const void>(InFrameDataD);
+				void* DataResult = StructProperty->ContainerPtrToValuePtr<void>(OutFrameData);
+				BilinearInterpolate(StructProperty->Struct, MainCoefficient, DeltaMinFocus, DeltaMaxFocus, DeltaMinZoom, DeltaMaxZoom, Data0, Data1, Data2, Data3, DataResult);
+			}
+		}
+		else if (FNumericProperty* NumericProperty = CastField<FNumericProperty>(Property))
+		{
+			if (NumericProperty->IsFloatingPoint())
+			{
+				const void* Data0 = NumericProperty->ContainerPtrToValuePtr<const void>(InFrameDataA);
+				double Value0 = NumericProperty->GetFloatingPointPropertyValue(Data0);
+				const void* Data1 = NumericProperty->ContainerPtrToValuePtr<const void>(InFrameDataB);
+				double Value1 = NumericProperty->GetFloatingPointPropertyValue(Data1);
+				const void* Data2 = NumericProperty->ContainerPtrToValuePtr<const void>(InFrameDataC);
+				double Value2 = NumericProperty->GetFloatingPointPropertyValue(Data2);
+				const void* Data3 = NumericProperty->ContainerPtrToValuePtr<const void>(InFrameDataD);
+				double Value3 = NumericProperty->GetFloatingPointPropertyValue(Data3);
+
+
+				double ValueResult = BilinearBlendValue(MainCoefficient, DeltaMinFocus, DeltaMaxFocus, DeltaMinZoom, DeltaMaxZoom, Value0, Value1, Value2, Value3);
+
+				void* DataResult = NumericProperty->ContainerPtrToValuePtr<void>(OutFrameData);
+				NumericProperty->SetFloatingPointPropertyValue(DataResult, ValueResult);
+			}
+			else if (NumericProperty->IsInteger() && !NumericProperty->IsEnum())
+			{
+				const void* Data0 = NumericProperty->ContainerPtrToValuePtr<const void>(InFrameDataA);
+				int64 Value0 = NumericProperty->GetSignedIntPropertyValue(Data0);
+				const void* Data1 = NumericProperty->ContainerPtrToValuePtr<const void>(InFrameDataB);
+				int64 Value1 = NumericProperty->GetSignedIntPropertyValue(Data1);
+				const void* Data2 = NumericProperty->ContainerPtrToValuePtr<const void>(InFrameDataC);
+				int64 Value2 = NumericProperty->GetSignedIntPropertyValue(Data2);
+				const void* Data3 = NumericProperty->ContainerPtrToValuePtr<const void>(InFrameDataD);
+				int64 Value3 = NumericProperty->GetSignedIntPropertyValue(Data3);
+
+				int64 ValueResult = BilinearBlendValue(MainCoefficient, DeltaMinFocus, DeltaMaxFocus, DeltaMinZoom, DeltaMaxZoom, Value0, Value1, Value2, Value3);
+
+				void* DataResult = NumericProperty->ContainerPtrToValuePtr<void>(OutFrameData);
+				NumericProperty->SetIntPropertyValue(DataResult, ValueResult);
+			}
+		}
+	}
+
+	template<typename Type>
+	void Interpolate(const FStructProperty* StructProperty, float InBlendWeight, const void* DataA, const void* DataB, void* DataResult)
+	{
+		const Type* ValuePtrA = StructProperty->ContainerPtrToValuePtr<Type>(DataA);
+		const Type* ValuePtrB = StructProperty->ContainerPtrToValuePtr<Type>(DataB);
+		Type* ValuePtrResult = StructProperty->ContainerPtrToValuePtr<Type>(DataResult);
+
+		Type ValueResult = BlendValue(InBlendWeight, *ValuePtrA, *ValuePtrB);
+		StructProperty->CopySingleValue(ValuePtrResult, &ValueResult);
+	}
+
+	void Interpolate(const UStruct* InStruct, float InBlendWeight, const void* InFrameDataA, const void* InFrameDataB, void* OutFrameData)
+	{
+		for (TFieldIterator<FProperty> Itt(InStruct); Itt; ++Itt)
+		{
+			FProperty* Property = *Itt;
+
+			if (FArrayProperty* ArrayProperty = CastField<FArrayProperty>(Property))
+			{
+				//ArrayProps have an ArrayDim of 1 but just to be sure...
+				for (int32 DimIndex = 0; DimIndex < ArrayProperty->ArrayDim; ++DimIndex)
+				{
+					const void* Data0 = ArrayProperty->ContainerPtrToValuePtr<const void>(InFrameDataA, DimIndex);
+					const void* Data1 = ArrayProperty->ContainerPtrToValuePtr<const void>(InFrameDataB, DimIndex);
+					void* DataResult = ArrayProperty->ContainerPtrToValuePtr<void>(OutFrameData, DimIndex);
+
+					FScriptArrayHelper ArrayHelperA(ArrayProperty, Data0);
+					FScriptArrayHelper ArrayHelperB(ArrayProperty, Data1);
+					FScriptArrayHelper ArrayHelperResult(ArrayProperty, DataResult);
+
+					const int32 MinValue = FMath::Min(ArrayHelperA.Num(), ArrayHelperB.Num());
+					ArrayHelperResult.Resize(MinValue);
+
+					for (int32 ArrayIndex = 0; ArrayIndex < MinValue; ++ArrayIndex)
+					{
+						InterpolateProperty(ArrayProperty->Inner, InBlendWeight, ArrayHelperA.GetRawPtr(ArrayIndex), ArrayHelperB.GetRawPtr(ArrayIndex), ArrayHelperResult.GetRawPtr(ArrayIndex));
+					}
+				}
+			}
+			else if (Property->ArrayDim > 1)
+			{
+				for (int32 DimIndex = 0; DimIndex < Property->ArrayDim; ++DimIndex)
+				{
+					const void* Data0 = Property->ContainerPtrToValuePtr<const void>(InFrameDataA, DimIndex);
+					const void* Data1 = Property->ContainerPtrToValuePtr<const void>(InFrameDataB, DimIndex);
+					void* DataResult = Property->ContainerPtrToValuePtr<void>(OutFrameData, DimIndex);
+
+					InterpolateProperty(Property, InBlendWeight, Data0, Data1, DataResult);
+				}
+			}
+			else
+			{
+				InterpolateProperty(Property, InBlendWeight, InFrameDataA, InFrameDataB, OutFrameData);
+			}
+		}
+	}
+
+	void InterpolateProperty(FProperty* Property, float InBlendWeight, const void* InDataA, const void* InDataB, void* OutData)
+	{
+		if (FStructProperty* StructProperty = CastField<FStructProperty>(Property))
+		{
+			if (StructProperty->Struct->GetFName() == NAME_Vector)
+			{
+				Interpolate<FVector>(StructProperty, InBlendWeight, InDataA, InDataB, OutData);
+			}
+			else if (StructProperty->Struct->GetFName() == NAME_Vector2D)
+			{
+				Interpolate<FVector2D>(StructProperty, InBlendWeight, InDataA, InDataB, OutData);
+			}
+			else if (StructProperty->Struct->GetFName() == NAME_Vector4)
+			{
+				Interpolate<FVector4>(StructProperty, InBlendWeight, InDataA, InDataB, OutData);
+			}
+			else if (StructProperty->Struct->GetFName() == NAME_Rotator)
+			{
+				Interpolate<FRotator>(StructProperty, InBlendWeight, InDataA, InDataB, OutData);
+			}
+			else if (StructProperty->Struct->GetFName() == NAME_Quat)
+			{
+				Interpolate<FQuat>(StructProperty, InBlendWeight, InDataA, InDataB, OutData);
+			}
+			else
+			{
+				const void* Data0 = StructProperty->ContainerPtrToValuePtr<const void>(InDataA);
+				const void* Data1 = StructProperty->ContainerPtrToValuePtr<const void>(InDataB);
+				void* DataResult = StructProperty->ContainerPtrToValuePtr<void>(OutData);
+				Interpolate(StructProperty->Struct, InBlendWeight, Data0, Data1, DataResult);
+			}
+		}
+		else if (FNumericProperty* NumericProperty = CastField<FNumericProperty>(Property))
+		{
+			if (NumericProperty->IsFloatingPoint())
+			{
+				const void* Data0 = NumericProperty->ContainerPtrToValuePtr<const void>(InDataA);
+				const double Value0 = NumericProperty->GetFloatingPointPropertyValue(Data0);
+				const void* Data1 = NumericProperty->ContainerPtrToValuePtr<const void>(InDataB);
+				const double Value1 = NumericProperty->GetFloatingPointPropertyValue(Data1);
+
+				const double ValueResult = BlendValue(InBlendWeight, Value0, Value1);
+
+				void* DataResult = NumericProperty->ContainerPtrToValuePtr<void>(OutData);
+				NumericProperty->SetFloatingPointPropertyValue(DataResult, ValueResult);
+			}
+			else if (NumericProperty->IsInteger() && !NumericProperty->IsEnum())
+			{
+				const void* Data0 = NumericProperty->ContainerPtrToValuePtr<const void>(InDataA);
+				const int64 Value0 = NumericProperty->GetSignedIntPropertyValue(Data0);
+				const void* Data1 = NumericProperty->ContainerPtrToValuePtr<const void>(InDataB);
+				const int64 Value1 = NumericProperty->GetSignedIntPropertyValue(Data1);
+
+				const int64 ValueResult = BlendValue(InBlendWeight, Value0, Value1);
+
+				void* DataResult = NumericProperty->ContainerPtrToValuePtr<void>(OutData);
+				NumericProperty->SetIntPropertyValue(DataResult, ValueResult);
+			}
+		}
+	}
+
+	float GetBlendFactor(float InValue, float ValueA, float ValueB)
+	{
+		//Keep input in range
+		InValue = FMath::Clamp(InValue, ValueA, ValueB);
+
+		const float Divider = ValueB - ValueA;
+		if (!FMath::IsNearlyZero(Divider))
+		{
+			return (InValue - ValueA) / Divider;
+		}
+		else
+		{
+			return 1.0f;
+		}
+	}
+}
+
+
diff --git a/Engine/Plugins/VirtualProduction/CameraCalibrationCore/Source/CameraCalibrationCore/Private/LensInterpolationUtils.h b/Engine/Plugins/VirtualProduction/CameraCalibrationCore/Source/CameraCalibrationCore/Private/LensInterpolationUtils.h
new file mode 100644
--- /dev/null
+++ b/Engine/Plugins/VirtualProduction/CameraCalibrationCore/Source/CameraCalibrationCore/Private/LensInterpolationUtils.h
@@ -0,0 +1,309 @@
+// Copyright Epic Games, Inc. All Rights Reserved.
+
+#pragma once
+
+#include "CoreTypes.h"
+
+#include "Misc/EnumClassFlags.h"
+#include "UObject/Object.h"
+#include "UObject/ObjectMacros.h"
+
+
+namespace LensInterpolationUtils
+{
+	//Flags for the possible direction on both axis
+	enum class EDirectionFlags : uint8
+	{
+		None = 0 << 0,
+		Negative = 1 << 0,
+		Positive = 1 << 1,
+	};
+	ENUM_CLASS_FLAGS(EDirectionFlags);
+
+	/** Helper struct when looking for four points around desired coord*/
+	struct FGridPointInfo
+	{
+		FGridPointInfo(EDirectionFlags InFocusDirection, EDirectionFlags InZoomDirection)
+		: NeededFocusDirection(InFocusDirection)
+		, NeededZoomDirection(InZoomDirection)
+		{}
+
+		const EDirectionFlags NeededFocusDirection;
+		const EDirectionFlags NeededZoomDirection;
+
+		//Score given to this coord based on direction that respects the desired coord.
+		//0: No direction respected
+		//1: One direction respected
+		//2:Two directions respected
+		int32 DirectionScore = 0;
+
+		//Index of that point in source data
+		int32 Index = INDEX_NONE;
+
+		//Cached value of that point coord (X = Focus, Y = Zoom)
+		float Focus = MAX_FLT;
+		float Zoom = MAX_FLT;
+	};
+
+	template<typename Type>
+	Type BilinearBlendValue(float MainCoefficient
+		, float DeltaMinFocus
+		, float DeltaMaxFocus
+		, float DeltaMinZoom
+		, float DeltaMaxZoom
+		, const Type& MinMin
+		, const Type& MinMax
+		, const Type& MaxMin
+		, const Type& MaxMax)
+	{
+		return (MinMin * DeltaMaxFocus * DeltaMaxZoom + MaxMin * DeltaMinFocus * DeltaMaxZoom + MinMax * DeltaMaxFocus * DeltaMinZoom + MaxMax * DeltaMinFocus * DeltaMinZoom) * MainCoefficient;
+	}
+
+
+	void BilinearInterpolate(const UStruct* InStruct
+		, float MainCoefficient
+		, float DeltaMinFocus
+		, float DeltaMaxFocus
+		, float DeltaMinZoom
+		, float DeltaMaxZoom
+		, const void* DataA
+		, const void* DataB
+		, const void* DataC
+		, const void* DataD
+		, void* OutFrameData);
+
+	template<typename Type>
+	void BilinearInterpolate(
+        float MainCoefficient
+        , float DeltaMinFocus
+        , float DeltaMaxFocus
+        , float DeltaMinZoom
+        , float DeltaMaxZoom
+        , const Type* DataA
+		, const Type* DataB
+		, const Type* DataC
+		, const Type* DataD
+		, Type* OutFrameData)
+	{
+		BilinearInterpolate(Type::StaticStruct(), MainCoefficient, DeltaMinFocus, DeltaMaxFocus, DeltaMinZoom, DeltaMaxZoom, DataA, DataB, DataC, DataD, OutFrameData);
+	}
+
+	template<typename Type>
+	Type BlendValue(float InBlendWeight, const Type& A, const Type& B)
+	{
+		return FMath::Lerp(A, B, InBlendWeight);
+	}
+	
+	void Interpolate(const UStruct* InStruct, float InBlendWeight, const void* InFrameDataA, const void* InFrameDataB, void* OutFrameData);
+
+	template<typename Type>
+	void Interpolate(float InBlendWeight, const Type* InFrameDataA, const Type* InFrameDataB, Type* OutFrameData)
+	{
+		Interpolate(Type::StaticStruct(), InBlendWeight, InFrameDataA, InFrameDataB, OutFrameData);
+	}
+	
+	float GetBlendFactor(float InValue, float ValueA, float ValueB);
+
+	/** Helper function to find the possibly four indices for the points around the desired coordinates */
+	template<typename Type>
+	bool FindBilinearInterpIndices(float Focus, float Zoom, TConstArrayView<Type> InSourceData, int32& OutPointIndexA, int32& OutPointIndexB, int32& OutPointIndexC, int32& OutPointIndexD)
+	{
+		if (InSourceData.Num() <= 0)
+		{
+			return false;
+		}
+
+		if (InSourceData.Num() == 1)
+		{
+			OutPointIndexA = OutPointIndexB = OutPointIndexC = OutPointIndexD = 0;
+			return true;
+		}
+
+		//To keep track of the four points we found
+		FGridPointInfo MinFocusMinZoom(EDirectionFlags::Negative, EDirectionFlags::Negative);
+		FGridPointInfo MinFocusMaxZoom(EDirectionFlags::Negative, EDirectionFlags::Positive);
+		FGridPointInfo MaxFocusMinZoom(EDirectionFlags::Positive, EDirectionFlags::Negative);
+		FGridPointInfo MaxFocusMaxZoom(EDirectionFlags::Positive, EDirectionFlags::Positive);
+
+		for (int32 Index = 0; Index < InSourceData.Num(); ++Index)
+		{
+			const Type& CurrentPoint = InSourceData[Index];
+			if (CurrentPoint.IsValid())
+			{
+				UpdatePointIfBetter(Focus, Zoom, CurrentPoint, Index, MinFocusMinZoom);
+				UpdatePointIfBetter(Focus, Zoom, CurrentPoint, Index, MinFocusMaxZoom);
+				UpdatePointIfBetter(Focus, Zoom, CurrentPoint, Index, MaxFocusMinZoom);
+				UpdatePointIfBetter(Focus, Zoom, CurrentPoint, Index, MaxFocusMaxZoom);
+			}
+		}
+
+		//Patch to fix bad point selection in corner case
+		//For some cases, the optimal solution is a line going across the desired coord.
+		//In its current form, the algorithm ends up selecting two identical points
+		//plus another one across it but the fourth one is in diagonal
+		//This will enforce that the point giving us a straight line is selected
+		if (((MinFocusMinZoom.Index == MinFocusMaxZoom.Index) && (MaxFocusMinZoom.Index != MaxFocusMaxZoom.Index)))
+		{
+			const Type& DuplicatePoint = InSourceData[MinFocusMinZoom.Index];
+			const Type& OtherPoint1 = InSourceData[MaxFocusMinZoom.Index];
+			const Type& OtherPoint2 = InSourceData[MaxFocusMaxZoom.Index];
+			const float DeltaFocus1 = FMath::Abs(OtherPoint1.Focus - DuplicatePoint.Focus);
+			const float DeltaFocus2 = FMath::Abs(OtherPoint2.Focus - DuplicatePoint.Focus);
+			const float DeltaZoom1 = FMath::Abs(OtherPoint1.Zoom - DuplicatePoint.Zoom);
+			const float DeltaZoom2 = FMath::Abs(OtherPoint2.Zoom - DuplicatePoint.Zoom);
+			if ((FMath::IsNearlyZero(DeltaFocus1) && !FMath::IsNearlyZero(DeltaZoom1))
+				||
+				(FMath::IsNearlyZero(DeltaZoom1) && !FMath::IsNearlyZero(DeltaFocus1)))
+			{
+				MaxFocusMaxZoom.Index = MaxFocusMinZoom.Index;
+			}
+			else if ((FMath::IsNearlyZero(DeltaFocus2) && !FMath::IsNearlyZero(DeltaZoom2))
+				||
+				(FMath::IsNearlyZero(DeltaZoom2) && !FMath::IsNearlyZero(DeltaFocus2)))
+			{
+				MaxFocusMinZoom.Index = MaxFocusMaxZoom.Index;
+			}
+		}
+		else if ((MinFocusMinZoom.Index == MaxFocusMinZoom.Index) && (MinFocusMaxZoom.Index != MaxFocusMaxZoom.Index))
+		{
+			//One point is badly selected. Maximize the line here
+			const Type& DuplicatePoint = InSourceData[MinFocusMinZoom.Index];
+			const Type& OtherPoint1 = InSourceData[MinFocusMaxZoom.Index];
+			const Type& OtherPoint2 = InSourceData[MaxFocusMaxZoom.Index];
+			const float DeltaFocus1 = FMath::Abs(OtherPoint1.Focus - DuplicatePoint.Focus);
+			const float DeltaFocus2 = FMath::Abs(OtherPoint2.Focus - DuplicatePoint.Focus);
+			const float DeltaZoom1 = FMath::Abs(OtherPoint1.Zoom - DuplicatePoint.Zoom);
+			const float DeltaZoom2 = FMath::Abs(OtherPoint2.Zoom - DuplicatePoint.Zoom);
+			if ((FMath::IsNearlyZero(DeltaFocus1) && !FMath::IsNearlyZero(DeltaZoom1))
+				||
+				(FMath::IsNearlyZero(DeltaZoom1) && !FMath::IsNearlyZero(DeltaFocus1)))
+			{
+				MaxFocusMaxZoom.Index = MinFocusMaxZoom.Index;
+			}
+			else if ((FMath::IsNearlyZero(DeltaFocus2) && !FMath::IsNearlyZero(DeltaZoom2))
+				||
+				(FMath::IsNearlyZero(DeltaZoom2) && !FMath::IsNearlyZero(DeltaFocus2)))
+			{
+				MinFocusMaxZoom.Index = MaxFocusMaxZoom.Index;
+			}
+		}
+
+		OutPointIndexA = MinFocusMinZoom.Index;
+		OutPointIndexB = MinFocusMaxZoom.Index;
+		OutPointIndexC = MaxFocusMinZoom.Index;
+		OutPointIndexD = MaxFocusMaxZoom.Index;
+
+		const bool bFoundValidResult = (OutPointIndexA != INDEX_NONE && OutPointIndexB != INDEX_NONE && OutPointIndexC != INDEX_NONE && OutPointIndexD != INDEX_NONE);
+		return bFoundValidResult;
+	}
+
+	template<typename Type>
+	void UpdatePointIfBetter(float Focus, float Zoom, const Type& NewPoint, int32 NewPointIndex, FGridPointInfo& CurrentPoint)
+	{
+		//Point is lower then desired coord so update MinMin if distance is smaller
+		const float NewDistance = FMath::Pow(FMath::Abs(Focus - NewPoint.Focus), 2) + FMath::Pow(FMath::Abs(Zoom - NewPoint.Zoom), 2);
+		const float CurrentDistance = FMath::Pow(FMath::Abs(Focus - CurrentPoint.Focus), 2) + FMath::Pow(FMath::Abs(Zoom - CurrentPoint.Zoom), 2);
+		const float NewFocusDir = NewPoint.Focus - Focus;
+		const float NewZoomDir = NewPoint.Zoom - Zoom;
+
+		const EDirectionFlags NewPointFocusDirectionFlags = ((NewFocusDir >= 0.0f) ? EDirectionFlags::Positive : EDirectionFlags::None) 
+													 | ((NewFocusDir <= 0.0f) ? EDirectionFlags::Negative : EDirectionFlags::None);
+		const EDirectionFlags NewPointZoomDirectionFlags = ((NewZoomDir >= 0.0f) ? EDirectionFlags::Positive : EDirectionFlags::None)
+													 | ((NewZoomDir <= 0.0f) ? EDirectionFlags::Negative : EDirectionFlags::None);
+
+		const int32 NewDirectionScore = (EnumHasAnyFlags(CurrentPoint.NeededFocusDirection, NewPointFocusDirectionFlags) ? 1 : 0)
+										+
+										(EnumHasAnyFlags(CurrentPoint.NeededZoomDirection, NewPointZoomDirectionFlags) ? 1 : 0);
+
+		bool bIsBetter = false;
+
+		//If new points improve required direction, take it
+		if (NewDirectionScore > CurrentPoint.DirectionScore)
+		{
+			bIsBetter = true;
+		}
+		else if(NewDirectionScore == CurrentPoint.DirectionScore)
+		{
+			//If new point has same score, favor closest one
+			bIsBetter = NewDistance < CurrentDistance;
+		}
+
+		if (bIsBetter)
+		{
+			CurrentPoint.DirectionScore = NewDirectionScore;
+			CurrentPoint.Index = NewPointIndex;
+			CurrentPoint.Focus = NewPoint.Focus;
+			CurrentPoint.Zoom = NewPoint.Zoom;
+		}
+	}
+
+	template<typename Type>
+	bool FIZMappingBilinearInterpolation(float InFocus, float InZoom, TConstArrayView<Type> InSourceData, Type& OutInterpolatedData)
+	{
+		int32 MinFocusMinZoomIndex = 0;
+		int32 MinFocusMaxZoomIndex = 0;
+		int32 MaxFocusMinZoomIndex = 0;
+		int32 MaxFocusMaxZoomIndex = 0;
+
+		//Start by finding the 4 points around the desired coords to do bilinear interpolation
+		//Depending on data map and desired coordinates, we might end up 
+		//  a. Not interpolating and returning a single entry point
+		//  b. Linear interpolate between two points
+		//  c. Bilinear interpolate between 4 points
+		const bool bFoundIndices = FindBilinearInterpIndices(InFocus, InZoom, InSourceData, MinFocusMinZoomIndex, MinFocusMaxZoomIndex, MaxFocusMinZoomIndex, MaxFocusMaxZoomIndex);
+		if (bFoundIndices)
+		{
+			check(InSourceData.IsValidIndex(MinFocusMinZoomIndex));
+			check(InSourceData.IsValidIndex(MinFocusMaxZoomIndex));
+			check(InSourceData.IsValidIndex(MaxFocusMinZoomIndex));
+			check(InSourceData.IsValidIndex(MaxFocusMaxZoomIndex));
+
+			const Type& MinFocusMinZoomPoint = InSourceData[MinFocusMinZoomIndex];
+			const Type& MinFocusMaxZoomPoint = InSourceData[MinFocusMaxZoomIndex];
+			const Type& MaxFocusMinZoomPoint = InSourceData[MaxFocusMinZoomIndex];
+			const Type& MaxFocusMaxZoomPoint = InSourceData[MaxFocusMaxZoomIndex];
+
+			//Single point case
+			if (MinFocusMinZoomIndex == MaxFocusMinZoomIndex && MaxFocusMinZoomIndex == MinFocusMaxZoomIndex && MinFocusMaxZoomIndex == MaxFocusMaxZoomIndex)
+			{
+				OutInterpolatedData = InSourceData[MinFocusMinZoomIndex];
+				return true;
+			}
+			else if (MinFocusMinZoomIndex == MaxFocusMinZoomIndex && MinFocusMaxZoomIndex == MaxFocusMaxZoomIndex)
+			{
+				//Fixed focus
+				const float BlendingFactor = GetBlendFactor(InZoom, MinFocusMinZoomPoint.Zoom, MaxFocusMaxZoomPoint.Zoom);
+				Interpolate(Type::StaticStruct(), BlendingFactor, &MinFocusMinZoomPoint, &MaxFocusMaxZoomPoint, &OutInterpolatedData);
+				return true;
+			}
+			else if (MinFocusMinZoomIndex == MinFocusMaxZoomIndex && MaxFocusMinZoomIndex == MaxFocusMaxZoomIndex)
+			{
+				//Fixed zoom
+				const float BlendingFactor = GetBlendFactor(InFocus, MinFocusMinZoomPoint.Focus, MaxFocusMaxZoomPoint.Focus);
+				Interpolate(Type::StaticStruct(), BlendingFactor, &MinFocusMinZoomPoint, &MaxFocusMaxZoomPoint, &OutInterpolatedData);
+				return true;
+			}
+			else
+			{
+				//The current grid finder doesn't always yield points around the sample
+				const float X2X1 = MaxFocusMinZoomPoint.Focus - MinFocusMinZoomPoint.Focus;
+				const float Y2Y1 = MaxFocusMaxZoomPoint.Zoom - MinFocusMinZoomPoint.Zoom;
+				const float Divider = X2X1 * Y2Y1;
+
+				if (!FMath::IsNearlyZero(Divider))
+				{
+					const float DeltaMaxFocus = MaxFocusMinZoomPoint.Focus - InFocus;
+					const float DeltaMaxZoom = MaxFocusMaxZoomPoint.Zoom - InZoom;
+					const float DeltaMinFocus = InFocus - MinFocusMinZoomPoint.Focus;
+					const float DeltaMinZoom = InZoom - MinFocusMinZoomPoint.Zoom;
+					const float MainCoeff = 1.0f / (X2X1 * Y2Y1);
+					BilinearInterpolate(Type::StaticStruct(), MainCoeff, DeltaMinFocus, DeltaMaxFocus, DeltaMinZoom, DeltaMaxZoom, &MinFocusMinZoomPoint, &MinFocusMaxZoomPoint, &MaxFocusMinZoomPoint, &MaxFocusMaxZoomPoint, &OutInterpolatedData);
+					return true;
+				}
+			}
+		}
+
+		return false;
+	}
+
+};
diff --git a/Engine/Plugins/VirtualProduction/CameraCalibrationCore/Source/CameraCalibrationCore/Private/Models/LensModel.cpp b/Engine/Plugins/VirtualProduction/CameraCalibrationCore/Source/CameraCalibrationCore/Private/Models/LensModel.cpp
new file mode 100644
--- /dev/null
+++ b/Engine/Plugins/VirtualProduction/CameraCalibrationCore/Source/CameraCalibrationCore/Private/Models/LensModel.cpp
@@ -0,0 +1,127 @@
+// Copyright Epic Games, Inc. All Rights Reserved.
+
+#include "Models/LensModel.h"
+
+
+#include "LensDistortionModelHandlerBase.h"
+#include "Logging/LogMacros.h"
+#include "UObject/UObjectIterator.h"
+
+
+DEFINE_LOG_CATEGORY_STATIC(LogLensModel, Log, All);
+
+#if WITH_EDITOR
+
+TArray<FText> ULensModel::GetParameterDisplayNames() const
+{
+	TArray<FText> ParameterNames;
+
+	UScriptStruct* TypeStruct = GetParameterStruct();
+
+	for (TFieldIterator<FProperty> It(TypeStruct); It; ++It)
+	{
+		if (FFloatProperty* Property = CastField<FFloatProperty>(*It))
+		{
+			ParameterNames.Add(Property->GetDisplayNameText());
+		}
+		else
+		{
+			UE_LOG(LogLensModel, Warning, TEXT("Property '%s' was skipped because its type was not float"), *(It->GetNameCPP()));
+		}
+	}
+
+	return ParameterNames;
+}
+
+#endif //WITH_EDITOR
+
+uint32 ULensModel::GetNumParameters() const
+{
+	UScriptStruct* TypeStruct = GetParameterStruct();
+
+	uint32 NumParameters = 0;
+	for (TFieldIterator<FProperty> It(TypeStruct); It; ++It)
+	{
+		if (FFloatProperty* Property = CastField<FFloatProperty>(*It))
+		{
+			++NumParameters;
+		}
+		else
+		{
+			UE_LOG(LogLensModel, Warning, TEXT("Property '%s' was skipped because its type was not float"), *(It->GetNameCPP()));
+		}
+	}
+
+	return NumParameters;
+}
+
+TSubclassOf<ULensDistortionModelHandlerBase> ULensModel::GetHandlerClass(TSubclassOf<ULensModel> LensModel)
+{
+	if (LensModel)
+	{
+		// Find all UClasses that derive from ULensDistortionModelHandlerBase
+		for (TObjectIterator<UClass> It; It; ++It)
+		{
+			if (It->IsChildOf(ULensDistortionModelHandlerBase::StaticClass()) && !It->HasAnyClassFlags(CLASS_Abstract | CLASS_Deprecated))
+			{
+				if (It->GetDefaultObject<ULensDistortionModelHandlerBase>()->IsModelSupported(LensModel))
+				{
+					return *It;
+				}
+			}
+		}
+	}
+	return nullptr;
+}
+
+void ULensModel::ToArray_Internal(const UScriptStruct* TypeStruct, const void* SrcData, TArray<float>& DstArray) const
+{
+	if (TypeStruct != GetParameterStruct())
+	{
+		UE_LOG(LogLensModel, Error, TEXT("TypeStruct does not match the distortion parameter struct supported by this model"));
+		return;
+	}
+
+	DstArray.Reset(GetNumParameters());
+	for (TFieldIterator<FProperty> It(TypeStruct); It; ++It)
+	{
+		if (FFloatProperty* Prop = CastField<FFloatProperty>(*It))
+		{
+			const float* Tmp = Prop->ContainerPtrToValuePtr<float>(SrcData);
+			DstArray.Add(*Tmp);
+		}
+		else
+		{
+			UE_LOG(LogLensModel, Warning, TEXT("Property '%s' was skipped because its type was not float"), *(It->GetNameCPP()));
+		}
+	}
+}
+
+void ULensModel::FromArray_Internal(UScriptStruct* TypeStruct, const TArray<float>& SrcArray, void* DstData)
+{
+	if (TypeStruct != GetParameterStruct())
+	{
+		UE_LOG(LogLensModel, Error, TEXT("TypeStruct does not match the distortion parameter struct supported by this model"));
+		return;
+	}
+
+	if (SrcArray.Num() != GetNumParameters())
+	{
+		UE_LOG(LogLensModel, Error, TEXT("SrcArray size (%d) does not match the expected number of parameters (%d)."), SrcArray.Num(), GetNumParameters());
+		return;
+	}
+
+	uint32 Index = 0;
+	for (TFieldIterator<FProperty> It(TypeStruct); It; ++It)
+	{
+		if (FFloatProperty* Prop = CastField<FFloatProperty>(*It))
+		{
+			Prop->SetPropertyValue_InContainer(DstData, SrcArray[Index]);
+			++Index;
+		}
+		else
+		{
+			UE_LOG(LogLensModel, Warning, TEXT("Property '%s' was skipped because its type was not float"), *(It->GetNameCPP()));
+		}
+	}
+}
diff --git a/Engine/Plugins/VirtualProduction/CameraCalibrationCore/Source/CameraCalibrationCore/Private/Models/SphericalLensModel.cpp b/Engine/Plugins/VirtualProduction/CameraCalibrationCore/Source/CameraCalibrationCore/Private/Models/SphericalLensModel.cpp
new file mode 100644
--- /dev/null
+++ b/Engine/Plugins/VirtualProduction/CameraCalibrationCore/Source/CameraCalibrationCore/Private/Models/SphericalLensModel.cpp
@@ -0,0 +1,13 @@
+// Copyright Epic Games, Inc. All Rights Reserved.
+
+#include "Models/SphericalLensModel.h"
+
+UScriptStruct* USphericalLensModel::GetParameterStruct() const
+{
+	return FSphericalDistortionParameters::StaticStruct();
+}
+
+FName USphericalLensModel::GetModelName() const 
+{ 
+	return FName("Spherical Lens Model"); 
+}
\ No newline at end of file
diff --git a/Engine/Plugins/VirtualProduction/CameraCalibrationCore/Source/CameraCalibrationCore/Private/SphericalLensDistortionModelHandler.cpp b/Engine/Plugins/VirtualProduction/CameraCalibrationCore/Source/CameraCalibrationCore/Private/SphericalLensDistortionModelHandler.cpp
new file mode 100644
--- /dev/null
+++ b/Engine/Plugins/VirtualProduction/CameraCalibrationCore/Source/CameraCalibrationCore/Private/SphericalLensDistortionModelHandler.cpp
@@ -0,0 +1,106 @@
+// Copyright Epic Games, Inc. All Rights Reserved.
+
+#include "SphericalLensDistortionModelHandler.h"
+
+#include "Engine/TextureRenderTarget2D.h"
+#include "CameraCalibrationSettings.h"
+#include "Math/NumericLimits.h"
+
+void USphericalLensDistortionModelHandler::InitializeHandler()
+{
+	LensModelClass = USphericalLensModel::StaticClass();
+}
+
+FVector2D USphericalLensDistortionModelHandler::ComputeDistortedUV(const FVector2D& InUndistortedUV) const
+{
+	// These distances cannot be zero in real-life. If they are, the current distortion state must be bad
+	if ((CurrentState.DistortionInfo.FxFy.X == 0.0f) || (CurrentState.DistortionInfo.FxFy.Y == 0.0f))
+	{
+		return InUndistortedUV;
+	}
+
+	FVector2D NormalizedDistanceFromImageCenter = (InUndistortedUV - CurrentState.PrincipalPoint) / CurrentState.DistortionInfo.FxFy;
+	const FVector2D OriginalDistance = NormalizedDistanceFromImageCenter;
+
+	// Iterative approach to distort an undistorted UV using coefficients that were designed to undistort
+	for (int32 Index = 0; Index < 2; ++Index)
+	{
+		const FVector2D DistanceSquared = NormalizedDistanceFromImageCenter * NormalizedDistanceFromImageCenter;
+		const float RSquared = DistanceSquared.X + DistanceSquared.Y;
+
+		const float RadialDistortion = 1.0f + (SphericalParameters.K1 * RSquared) + (SphericalParameters.K2 * RSquared * RSquared) + (SphericalParameters.K3 * RSquared * RSquared * RSquared);
+
+		const FVector2D TangentialDistortion = 
+		{
+			(SphericalParameters.P2 * (RSquared + (2.0f * DistanceSquared.X))) + (2.0f * SphericalParameters.P1 * NormalizedDistanceFromImageCenter.X * NormalizedDistanceFromImageCenter.Y),
+			(SphericalParameters.P1 * (RSquared + (2.0f * DistanceSquared.Y))) + (2.0f * SphericalParameters.P2 * NormalizedDistanceFromImageCenter.X * NormalizedDistanceFromImageCenter.Y)
+		};
+
+		// Guard against divide-by-zero errors
+		if (RadialDistortion == 0.0f)
+		{
+			NormalizedDistanceFromImageCenter = FVector2D(TNumericLimits<float>::Max(), TNumericLimits<float>::Max());
+			break;
+		}
+		else
+		{
+			NormalizedDistanceFromImageCenter = (OriginalDistance - TangentialDistortion) / RadialDistortion;
+		}
+	}
+
+	const FVector2D DistortedUV = (NormalizedDistanceFromImageCenter * CurrentState.DistortionInfo.FxFy) + FVector2D(0.5f, 0.5f);
+	return DistortedUV;
+}
+
+void USphericalLensDistortionModelHandler::InitDistortionMaterials()
+{
+	if (DistortionPostProcessMID == nullptr)
+	{
+		UMaterialInterface* DistortionMaterialParent = GetDefault<UCameraCalibrationSettings>()->GetDefaultDistortionMaterial(this->StaticClass());
+		DistortionPostProcessMID = UMaterialInstanceDynamic::Create(DistortionMaterialParent, this);
+	}
+
+	if (UndistortionDisplacementMapMID == nullptr)
+	{
+		UMaterialInterface* MaterialParent = GetDefault<UCameraCalibrationSettings>()->GetDefaultUndistortionDisplacementMaterial(this->StaticClass());
+		UndistortionDisplacementMapMID = UMaterialInstanceDynamic::Create(MaterialParent, this);
+	}
+
+	if (DistortionDisplacementMapMID == nullptr)
+	{
+		UMaterialInterface* MaterialParent = GetDefault<UCameraCalibrationSettings>()->GetDefaultDistortionDisplacementMaterial(this->StaticClass());
+		DistortionDisplacementMapMID = UMaterialInstanceDynamic::Create(MaterialParent, this);
+	}
+
+	DistortionPostProcessMID->SetTextureParameterValue("UndistortionDisplacementMap", UndistortionDisplacementMapRT);
+	DistortionPostProcessMID->SetTextureParameterValue("DistortionDisplacementMap", DistortionDisplacementMapRT);
+
+	SetDistortionState(CurrentState);
+}
+
+void USphericalLensDistortionModelHandler::UpdateMaterialParameters()
+{
+	//Helper function to set material parameters of an MID
+	const auto SetDistortionMaterialParameters = [this](UMaterialInstanceDynamic* const MID)
+	{
+		MID->SetScalarParameterValue("k1", SphericalParameters.K1);
+		MID->SetScalarParameterValue("k2", SphericalParameters.K2);
+		MID->SetScalarParameterValue("k3", SphericalParameters.K3);
+		MID->SetScalarParameterValue("p1", SphericalParameters.P1);
+		MID->SetScalarParameterValue("p2", SphericalParameters.P2);
+
+		MID->SetScalarParameterValue("cx", CurrentState.PrincipalPoint.X);
+		MID->SetScalarParameterValue("cy", CurrentState.PrincipalPoint.Y);
+
+		MID->SetScalarParameterValue("fx", CurrentState.DistortionInfo.FxFy.X);
+		MID->SetScalarParameterValue("fy", CurrentState.DistortionInfo.FxFy.Y);
+	};
+
+	SetDistortionMaterialParameters(UndistortionDisplacementMapMID);
+	SetDistortionMaterialParameters(DistortionDisplacementMapMID);
+}
+
+void USphericalLensDistortionModelHandler::InterpretDistortionParameters()
+{
+	LensModelClass->GetDefaultObject<ULensModel>()->FromArray<FSphericalDistortionParameters>(CurrentState.DistortionInfo.Parameters, SphericalParameters);
+}
diff --git a/Engine/Plugins/VirtualProduction/CameraCalibrationCore/Source/CameraCalibrationCore/Private/Tests/TestDistortionMapping.cpp b/Engine/Plugins/VirtualProduction/CameraCalibrationCore/Source/CameraCalibrationCore/Private/Tests/TestDistortionMapping.cpp
new file mode 100644
--- /dev/null
+++ b/Engine/Plugins/VirtualProduction/CameraCalibrationCore/Source/CameraCalibrationCore/Private/Tests/TestDistortionMapping.cpp
@@ -0,0 +1,320 @@
+// Copyright Epic Games, Inc. All Rights Reserved.
+
+#if WITH_DEV_AUTOMATION_TESTS
+
+#include "Misc/AutomationTest.h"
+
+#include "CoreMinimal.h"
+#include "LensFile.h"
+#include "LensInterpolationUtils.h"
+
+IMPLEMENT_SIMPLE_AUTOMATION_TEST(FTestCameraCalibrationCore, "Plugins.CameraCalibrationCore", EAutomationTestFlags::EditorContext | EAutomationTestFlags::EngineFilter)
+
+namespace CameraCalibrationTestUtil
+{
+	void TestBilinearFindGrid(FAutomationTestBase& Test)
+	{
+		TArray<FDistortionMapPoint> DistortionMapping;
+
+		FDistortionMapPoint TempPoint;
+		TempPoint.Focus = 0.0f;
+		TempPoint.Zoom = 0.0f;
+
+		float MaxFocus = 100.0f;
+		float MaxZoom = 100.0f;
+		int32 FocusDiv = 10;
+		int32 ZoomDiv = 10;
+
+		//Asking for indices should fail with empty array
+		int32 MinMinPoint, MinMaxPoint, MaxMinPoint, MaxMaxPoint;
+		float DesiredFocus = 0.0f;
+		float DesiredZoom = 0.0f;
+		bool bSuccess = LensInterpolationUtils::FindBilinearInterpIndices<FDistortionMapPoint>(DesiredFocus, DesiredZoom, DistortionMapping, MinMinPoint, MinMaxPoint, MaxMinPoint, MaxMaxPoint);
+		Test.TestEqual(TEXT("FindInterp - Empty"), bSuccess, false);
+
+		DistortionMapping.Reserve(FocusDiv * ZoomDiv);
+		for (int32 FocusIndex = 0; FocusIndex < FocusDiv; ++FocusIndex)
+		{
+			const float FocusValue = MaxFocus / (float)FocusDiv * (float)FocusIndex;
+
+			if (FocusIndex != 7)
+			{
+				for (int32 ZoomIndex = 0; ZoomIndex < ZoomDiv; ++ZoomIndex)
+				{
+					const float ZoomValue = MaxZoom / (float)ZoomDiv * (float)ZoomIndex;
+					TempPoint.Focus = FocusValue;
+					TempPoint.Zoom = ZoomValue;
+					TempPoint.DistortionInfo.Parameters[0] = DistortionMapping.Num();
+					DistortionMapping.Add(TempPoint);
+				}
+			}
+			else
+			{
+				//Skip points for zoom
+				for (int32 ZoomIndex = 0; ZoomIndex < 2; ++ZoomIndex)
+				{
+					const float ZoomValue = MaxZoom / (float)ZoomDiv * (float)ZoomIndex;
+					TempPoint.Focus = FocusValue;
+					TempPoint.Zoom = ZoomValue;
+					TempPoint.DistortionInfo.Parameters[0] = DistortionMapping.Num();
+					DistortionMapping.Add(TempPoint);
+				}
+			}
+		}
+
+		Test.TestEqual(TEXT("FindInterp - Empty"), bSuccess, false);
+
+		
+		FDistortionMapPoint Result;
+
+		{
+			DesiredFocus = 5.0f;
+			DesiredZoom = 5.0f;
+			bSuccess = LensInterpolationUtils::FindBilinearInterpIndices<FDistortionMapPoint>(DesiredFocus, DesiredZoom, DistortionMapping, MinMinPoint, MinMaxPoint, MaxMinPoint, MaxMaxPoint);
+			LensInterpolationUtils::FIZMappingBilinearInterpolation<FDistortionMapPoint>(DesiredFocus, DesiredZoom, DistortionMapping, Result);
+
+			Test.TestEqual(*FString::Printf(TEXT("Find Indices - (%0.2f,%0.2f)"), DesiredFocus, DesiredZoom), bSuccess, true);
+			Test.TestEqual(*FString::Printf(TEXT("MinMin - (%0.2f,%0.2f)"), DesiredFocus, DesiredZoom), MinMinPoint, 0);
+			Test.TestEqual(*FString::Printf(TEXT("MinMax - (%0.2f,%0.2f)"), DesiredFocus, DesiredZoom), MinMaxPoint, 1);
+			Test.TestEqual(*FString::Printf(TEXT("MaxMin - (%0.2f,%0.2f)"), DesiredFocus, DesiredZoom), MaxMinPoint, 10);
+			Test.TestEqual(*FString::Printf(TEXT("MaxMax - (%0.2f,%0.2f)"), DesiredFocus, DesiredZoom), MaxMaxPoint, 11);
+			Test.TestEqual(*FString::Printf(TEXT("InterpResult - (%0.2f,%0.2f)"), DesiredFocus, DesiredZoom), Result.DistortionInfo.Parameters[0], 5.5f);
+		}
+		
+		{
+			DesiredFocus = 2.5f;
+			DesiredZoom = 5.0f;
+			bSuccess = LensInterpolationUtils::FindBilinearInterpIndices<FDistortionMapPoint>(DesiredFocus, DesiredZoom, DistortionMapping, MinMinPoint, MinMaxPoint, MaxMinPoint, MaxMaxPoint);
+			LensInterpolationUtils::FIZMappingBilinearInterpolation<FDistortionMapPoint>(DesiredFocus, DesiredZoom, DistortionMapping, Result);
+
+			Test.TestEqual(*FString::Printf(TEXT("Find Indices - (%0.2f,%0.2f)"), DesiredFocus, DesiredZoom), bSuccess, true);
+			Test.TestEqual(*FString::Printf(TEXT("MinMin - (%0.2f,%0.2f)"), DesiredFocus, DesiredZoom), MinMinPoint, 0);
+			Test.TestEqual(*FString::Printf(TEXT("MinMax - (%0.2f,%0.2f)"), DesiredFocus, DesiredZoom), MinMaxPoint, 1);
+			Test.TestEqual(*FString::Printf(TEXT("MaxMin - (%0.2f,%0.2f)"), DesiredFocus, DesiredZoom), MaxMinPoint, 10);
+			Test.TestEqual(*FString::Printf(TEXT("MaxMax - (%0.2f,%0.2f)"), DesiredFocus, DesiredZoom), MaxMaxPoint, 11);
+			Test.TestEqual(*FString::Printf(TEXT("InterpResult - (%0.2f,%0.2f)"), DesiredFocus, DesiredZoom), Result.DistortionInfo.Parameters[0], 3.0f);
+		}
+	
+
+		{
+			DesiredFocus = 0.0f;
+			DesiredZoom = 5.0f;
+			bSuccess = LensInterpolationUtils::FindBilinearInterpIndices<FDistortionMapPoint>(DesiredFocus, DesiredZoom, DistortionMapping, MinMinPoint, MinMaxPoint, MaxMinPoint, MaxMaxPoint);
+			LensInterpolationUtils::FIZMappingBilinearInterpolation<FDistortionMapPoint>(DesiredFocus, DesiredZoom, DistortionMapping, Result);
+
+			Test.TestEqual(*FString::Printf(TEXT("Find Indices - (%0.2f,%0.2f)"), DesiredFocus, DesiredZoom), bSuccess, true);
+			Test.TestEqual(*FString::Printf(TEXT("MinMin - (%0.2f,%0.2f)"), DesiredFocus, DesiredZoom), MinMinPoint, 0);
+			Test.TestEqual(*FString::Printf(TEXT("MinMax - (%0.2f,%0.2f)"), DesiredFocus, DesiredZoom), MinMaxPoint, 1);
+			Test.TestEqual(*FString::Printf(TEXT("MaxMin - (%0.2f,%0.2f)"), DesiredFocus, DesiredZoom), MaxMinPoint, 0);
+			Test.TestEqual(*FString::Printf(TEXT("MaxMax - (%0.2f,%0.2f)"), DesiredFocus, DesiredZoom), MaxMaxPoint, 1);
+			Test.TestEqual(*FString::Printf(TEXT("InterpResult - (%0.2f,%0.2f)"), DesiredFocus, DesiredZoom), Result.DistortionInfo.Parameters[0], 0.5f);
+		}
+
+		{
+			DesiredFocus = 5.0f;
+			DesiredZoom = 0.0f;
+			bSuccess = LensInterpolationUtils::FindBilinearInterpIndices<FDistortionMapPoint>(DesiredFocus, DesiredZoom, DistortionMapping, MinMinPoint, MinMaxPoint, MaxMinPoint, MaxMaxPoint);
+			LensInterpolationUtils::FIZMappingBilinearInterpolation<FDistortionMapPoint>(DesiredFocus, DesiredZoom, DistortionMapping, Result);
+
+			Test.TestEqual(*FString::Printf(TEXT("Find Indices - (%0.2f,%0.2f)"), DesiredFocus, DesiredZoom), bSuccess, true);
+			Test.TestEqual(*FString::Printf(TEXT("MinMin - (%0.2f,%0.2f)"), DesiredFocus, DesiredZoom), MinMinPoint, 0);
+			Test.TestEqual(*FString::Printf(TEXT("MinMax - (%0.2f,%0.2f)"), DesiredFocus, DesiredZoom), MinMaxPoint, 0);
+			Test.TestEqual(*FString::Printf(TEXT("MaxMin - (%0.2f,%0.2f)"), DesiredFocus, DesiredZoom), MaxMinPoint, 10);
+			Test.TestEqual(*FString::Printf(TEXT("MaxMax - (%0.2f,%0.2f)"), DesiredFocus, DesiredZoom), MaxMaxPoint, 10);
+			Test.TestEqual(*FString::Printf(TEXT("InterpResult - (%0.2f,%0.2f)"), DesiredFocus, DesiredZoom), Result.DistortionInfo.Parameters[0], 5.0f);
+		}
+		
+		{
+			DesiredFocus = 30.0f;
+			DesiredZoom = 10.0f;
+			bSuccess = LensInterpolationUtils::FindBilinearInterpIndices<FDistortionMapPoint>(DesiredFocus, DesiredZoom, DistortionMapping, MinMinPoint, MinMaxPoint, MaxMinPoint, MaxMaxPoint);
+			LensInterpolationUtils::FIZMappingBilinearInterpolation<FDistortionMapPoint>(DesiredFocus, DesiredZoom, DistortionMapping, Result);
+
+			Test.TestEqual(*FString::Printf(TEXT("Find Indices - (%0.2f,%0.2f)"), DesiredFocus, DesiredZoom), bSuccess, true);
+			Test.TestEqual(*FString::Printf(TEXT("MinMin - (%0.2f,%0.2f)"), DesiredFocus, DesiredZoom), MinMinPoint, 31);
+			Test.TestEqual(*FString::Printf(TEXT("MinMax - (%0.2f,%0.2f)"), DesiredFocus, DesiredZoom), MinMaxPoint, 31);
+			Test.TestEqual(*FString::Printf(TEXT("MaxMin - (%0.2f,%0.2f)"), DesiredFocus, DesiredZoom), MaxMinPoint, 31);
+			Test.TestEqual(*FString::Printf(TEXT("MaxMax - (%0.2f,%0.2f)"), DesiredFocus, DesiredZoom), MaxMaxPoint, 31);
+			Test.TestEqual(*FString::Printf(TEXT("InterpResult - (%0.2f,%0.2f)"), DesiredFocus, DesiredZoom), Result.DistortionInfo.Parameters[0], 31.0f);
+		}
+
+		{
+			DesiredFocus = 70.0f;
+			DesiredZoom = 20.0f;
+			bSuccess = LensInterpolationUtils::FindBilinearInterpIndices<FDistortionMapPoint>(DesiredFocus, DesiredZoom, DistortionMapping, MinMinPoint, MinMaxPoint, MaxMinPoint, MaxMaxPoint);
+			LensInterpolationUtils::FIZMappingBilinearInterpolation<FDistortionMapPoint>(DesiredFocus, DesiredZoom, DistortionMapping, Result);
+
+			Test.TestEqual(*FString::Printf(TEXT("Find Indices - (%0.2f,%0.2f)"), DesiredFocus, DesiredZoom), bSuccess, true);
+			Test.TestEqual(*FString::Printf(TEXT("MinMin - (%0.2f,%0.2f)"), DesiredFocus, DesiredZoom), MinMinPoint, 62);
+			Test.TestEqual(*FString::Printf(TEXT("MinMax - (%0.2f,%0.2f)"), DesiredFocus, DesiredZoom), MinMaxPoint, 62);
+			Test.TestEqual(*FString::Printf(TEXT("MaxMin - (%0.2f,%0.2f)"), DesiredFocus, DesiredZoom), MaxMinPoint, 74);
+			Test.TestEqual(*FString::Printf(TEXT("MaxMax - (%0.2f,%0.2f)"), DesiredFocus, DesiredZoom), MaxMaxPoint, 74);
+			Test.TestEqual(*FString::Printf(TEXT("InterpResult - (%0.2f,%0.2f)"), DesiredFocus, DesiredZoom), Result.DistortionInfo.Parameters[0], 68.0f);
+		}
+
+		{
+			DesiredFocus = 70.0f;
+			DesiredZoom = 5.0f;
+			bSuccess = LensInterpolationUtils::FindBilinearInterpIndices<FDistortionMapPoint>(DesiredFocus, DesiredZoom, DistortionMapping, MinMinPoint, MinMaxPoint, MaxMinPoint, MaxMaxPoint);
+			LensInterpolationUtils::FIZMappingBilinearInterpolation<FDistortionMapPoint>(DesiredFocus, DesiredZoom, DistortionMapping, Result);
+
+			Test.TestEqual(*FString::Printf(TEXT("Find Indices - (%0.2f,%0.2f)"), DesiredFocus, DesiredZoom), bSuccess, true);
+			Test.TestEqual(*FString::Printf(TEXT("MinMin - (%0.2f,%0.2f)"), DesiredFocus, DesiredZoom), MinMinPoint, 70);
+			Test.TestEqual(*FString::Printf(TEXT("MinMax - (%0.2f,%0.2f)"), DesiredFocus, DesiredZoom), MinMaxPoint, 71);
+			Test.TestEqual(*FString::Printf(TEXT("MaxMin - (%0.2f,%0.2f)"), DesiredFocus, DesiredZoom), MaxMinPoint, 70);
+			Test.TestEqual(*FString::Printf(TEXT("MaxMax - (%0.2f,%0.2f)"), DesiredFocus, DesiredZoom), MaxMaxPoint, 71);
+			Test.TestEqual(*FString::Printf(TEXT("InterpResult - (%0.2f,%0.2f)"), DesiredFocus, DesiredZoom), Result.DistortionInfo.Parameters[0], 70.5f);
+		}
+
+		{
+			DesiredFocus = 70.0f;
+			DesiredZoom = 10.0f;
+			bSuccess = LensInterpolationUtils::FindBilinearInterpIndices<FDistortionMapPoint>(DesiredFocus, DesiredZoom, DistortionMapping, MinMinPoint, MinMaxPoint, MaxMinPoint, MaxMaxPoint);
+			LensInterpolationUtils::FIZMappingBilinearInterpolation<FDistortionMapPoint>(DesiredFocus, DesiredZoom, DistortionMapping, Result);
+
+			Test.TestEqual(*FString::Printf(TEXT("Find Indices - (%0.2f,%0.2f)"), DesiredFocus, DesiredZoom), bSuccess, true);
+			Test.TestEqual(*FString::Printf(TEXT("MinMin - (%0.2f,%0.2f)"), DesiredFocus, DesiredZoom), MinMinPoint, 71);
+			Test.TestEqual(*FString::Printf(TEXT("MinMax - (%0.2f,%0.2f)"), DesiredFocus, DesiredZoom), MinMaxPoint, 71);
+			Test.TestEqual(*FString::Printf(TEXT("MaxMin - (%0.2f,%0.2f)"), DesiredFocus, DesiredZoom), MaxMinPoint, 71);
+			Test.TestEqual(*FString::Printf(TEXT("MaxMax - (%0.2f,%0.2f)"), DesiredFocus, DesiredZoom), MaxMaxPoint, 71);
+			Test.TestEqual(*FString::Printf(TEXT("InterpResult - (%0.2f,%0.2f)"), DesiredFocus, DesiredZoom), Result.DistortionInfo.Parameters[0], 71.0f);
+		}
+
+		{
+			DesiredFocus = 70.0f;
+			DesiredZoom = 25.0f;
+			bSuccess = LensInterpolationUtils::FindBilinearInterpIndices<FDistortionMapPoint>(DesiredFocus, DesiredZoom, DistortionMapping, MinMinPoint, MinMaxPoint, MaxMinPoint, MaxMaxPoint);
+			LensInterpolationUtils::FIZMappingBilinearInterpolation<FDistortionMapPoint>(DesiredFocus, DesiredZoom, DistortionMapping, Result);
+
+			Test.TestEqual(*FString::Printf(TEXT("Find Indices - (%0.2f,%0.2f)"), DesiredFocus, DesiredZoom), bSuccess, true);
+			Test.TestEqual(*FString::Printf(TEXT("MinMin - (%0.2f,%0.2f)"), DesiredFocus, DesiredZoom), MinMinPoint, 62);
+			Test.TestEqual(*FString::Printf(TEXT("MinMax - (%0.2f,%0.2f)"), DesiredFocus, DesiredZoom), MinMaxPoint, 63);
+			Test.TestEqual(*FString::Printf(TEXT("MaxMin - (%0.2f,%0.2f)"), DesiredFocus, DesiredZoom), MaxMinPoint, 74);
+			Test.TestEqual(*FString::Printf(TEXT("MaxMax - (%0.2f,%0.2f)"), DesiredFocus, DesiredZoom), MaxMaxPoint, 75);
+			Test.TestEqual(*FString::Printf(TEXT("InterpResult - (%0.2f,%0.2f)"), DesiredFocus, DesiredZoom), Result.DistortionInfo.Parameters[0], 68.5f);
+		}
+
+		{
+			DesiredFocus = -1.0f;
+			DesiredZoom = -1.0f;
+			bSuccess = LensInterpolationUtils::FindBilinearInterpIndices<FDistortionMapPoint>(DesiredFocus, DesiredZoom, DistortionMapping, MinMinPoint, MinMaxPoint, MaxMinPoint, MaxMaxPoint);
+			LensInterpolationUtils::FIZMappingBilinearInterpolation<FDistortionMapPoint>(DesiredFocus, DesiredZoom, DistortionMapping, Result);
+
+			Test.TestEqual(*FString::Printf(TEXT("Find Indices - (%0.2f,%0.2f)"), DesiredFocus, DesiredZoom), bSuccess, true);
+			Test.TestEqual(*FString::Printf(TEXT("MinMin - (%0.2f,%0.2f)"), DesiredFocus, DesiredZoom), MinMinPoint, 0);
+			Test.TestEqual(*FString::Printf(TEXT("MinMax - (%0.2f,%0.2f)"), DesiredFocus, DesiredZoom), MinMaxPoint, 0);
+			Test.TestEqual(*FString::Printf(TEXT("MaxMin - (%0.2f,%0.2f)"), DesiredFocus, DesiredZoom), MaxMinPoint, 0);
+			Test.TestEqual(*FString::Printf(TEXT("MaxMax - (%0.2f,%0.2f)"), DesiredFocus, DesiredZoom), MaxMaxPoint, 0);
+			Test.TestEqual(*FString::Printf(TEXT("InterpResult - (%0.2f,%0.2f)"), DesiredFocus, DesiredZoom), Result.DistortionInfo.Parameters[0], 0.0f);
+		}
+
+		{
+			DesiredFocus = -1.0f;
+			DesiredZoom = 35.0f;
+			bSuccess = LensInterpolationUtils::FindBilinearInterpIndices<FDistortionMapPoint>(DesiredFocus, DesiredZoom, DistortionMapping, MinMinPoint, MinMaxPoint, MaxMinPoint, MaxMaxPoint);
+			LensInterpolationUtils::FIZMappingBilinearInterpolation<FDistortionMapPoint>(DesiredFocus, DesiredZoom, DistortionMapping, Result);
+
+			Test.TestEqual(*FString::Printf(TEXT("Find Indices - (%0.2f,%0.2f)"), DesiredFocus, DesiredZoom), bSuccess, true);
+			Test.TestEqual(*FString::Printf(TEXT("MinMin - (%0.2f,%0.2f)"), DesiredFocus, DesiredZoom), MinMinPoint, 3);
+			Test.TestEqual(*FString::Printf(TEXT("MinMax - (%0.2f,%0.2f)"), DesiredFocus, DesiredZoom), MinMaxPoint, 4);
+			Test.TestEqual(*FString::Printf(TEXT("MaxMin - (%0.2f,%0.2f)"), DesiredFocus, DesiredZoom), MaxMinPoint, 3);
+			Test.TestEqual(*FString::Printf(TEXT("MaxMax - (%0.2f,%0.2f)"), DesiredFocus, DesiredZoom), MaxMaxPoint, 4);
+			Test.TestEqual(*FString::Printf(TEXT("InterpResult - (%0.2f,%0.2f)"), DesiredFocus, DesiredZoom), Result.DistortionInfo.Parameters[0], 3.5f);
+		}
+
+		{
+			DesiredFocus = -1.0f;
+			DesiredZoom = 110.0f;
+			bSuccess = LensInterpolationUtils::FindBilinearInterpIndices<FDistortionMapPoint>(DesiredFocus, DesiredZoom, DistortionMapping, MinMinPoint, MinMaxPoint, MaxMinPoint, MaxMaxPoint);
+			LensInterpolationUtils::FIZMappingBilinearInterpolation<FDistortionMapPoint>(DesiredFocus, DesiredZoom, DistortionMapping, Result);
+
+			Test.TestEqual(*FString::Printf(TEXT("Find Indices - (%0.2f,%0.2f)"), DesiredFocus, DesiredZoom), bSuccess, true);
+			Test.TestEqual(*FString::Printf(TEXT("MinMin - (%0.2f,%0.2f)"), DesiredFocus, DesiredZoom), MinMinPoint, 9);
+			Test.TestEqual(*FString::Printf(TEXT("MinMax - (%0.2f,%0.2f)"), DesiredFocus, DesiredZoom), MinMaxPoint, 9);
+			Test.TestEqual(*FString::Printf(TEXT("MaxMin - (%0.2f,%0.2f)"), DesiredFocus, DesiredZoom), MaxMinPoint, 9);
+			Test.TestEqual(*FString::Printf(TEXT("MaxMax - (%0.2f,%0.2f)"), DesiredFocus, DesiredZoom), MaxMaxPoint, 9);
+			Test.TestEqual(*FString::Printf(TEXT("InterpResult - (%0.2f,%0.2f)"), DesiredFocus, DesiredZoom), Result.DistortionInfo.Parameters[0], 9.0f);
+		}
+
+		{
+			DesiredFocus = 35.0f;
+			DesiredZoom = 110.0f;
+			bSuccess = LensInterpolationUtils::FindBilinearInterpIndices<FDistortionMapPoint>(DesiredFocus, DesiredZoom, DistortionMapping, MinMinPoint, MinMaxPoint, MaxMinPoint, MaxMaxPoint);
+			LensInterpolationUtils::FIZMappingBilinearInterpolation<FDistortionMapPoint>(DesiredFocus, DesiredZoom, DistortionMapping, Result);
+
+			Test.TestEqual(*FString::Printf(TEXT("Find Indices - (%0.2f,%0.2f)"), DesiredFocus, DesiredZoom), bSuccess, true);
+			Test.TestEqual(*FString::Printf(TEXT("MinMin - (%0.2f,%0.2f)"), DesiredFocus, DesiredZoom), MinMinPoint, 39);
+			Test.TestEqual(*FString::Printf(TEXT("MinMax - (%0.2f,%0.2f)"), DesiredFocus, DesiredZoom), MinMaxPoint, 39);
+			Test.TestEqual(*FString::Printf(TEXT("MaxMin - (%0.2f,%0.2f)"), DesiredFocus, DesiredZoom), MaxMinPoint, 49);
+			Test.TestEqual(*FString::Printf(TEXT("MaxMax - (%0.2f,%0.2f)"), DesiredFocus, DesiredZoom), MaxMaxPoint, 49);
+			Test.TestEqual(*FString::Printf(TEXT("InterpResult - (%0.2f,%0.2f)"), DesiredFocus, DesiredZoom), Result.DistortionInfo.Parameters[0], 44.0f);
+		}
+
+		{
+			DesiredFocus = 110.0f;
+			DesiredZoom = 110.0f;
+			bSuccess = LensInterpolationUtils::FindBilinearInterpIndices<FDistortionMapPoint>(DesiredFocus, DesiredZoom, DistortionMapping, MinMinPoint, MinMaxPoint, MaxMinPoint, MaxMaxPoint);
+			LensInterpolationUtils::FIZMappingBilinearInterpolation<FDistortionMapPoint>(DesiredFocus, DesiredZoom, DistortionMapping, Result);
+
+			Test.TestEqual(*FString::Printf(TEXT("Find Indices - (%0.2f,%0.2f)"), DesiredFocus, DesiredZoom), bSuccess, true);
+			Test.TestEqual(*FString::Printf(TEXT("MinMin - (%0.2f,%0.2f)"), DesiredFocus, DesiredZoom), MinMinPoint, 91);
+			Test.TestEqual(*FString::Printf(TEXT("MinMax - (%0.2f,%0.2f)"), DesiredFocus, DesiredZoom), MinMaxPoint, 91);
+			Test.TestEqual(*FString::Printf(TEXT("MaxMin - (%0.2f,%0.2f)"), DesiredFocus, DesiredZoom), MaxMinPoint, 91);
+			Test.TestEqual(*FString::Printf(TEXT("MaxMax - (%0.2f,%0.2f)"), DesiredFocus, DesiredZoom), MaxMaxPoint, 91);
+			Test.TestEqual(*FString::Printf(TEXT("InterpResult - (%0.2f,%0.2f)"), DesiredFocus, DesiredZoom), Result.DistortionInfo.Parameters[0], 91.0f);
+		}
+
+		{
+			DesiredFocus = 110.0f;
+			DesiredZoom = 35.0f;
+			bSuccess = LensInterpolationUtils::FindBilinearInterpIndices<FDistortionMapPoint>(DesiredFocus, DesiredZoom, DistortionMapping, MinMinPoint, MinMaxPoint, MaxMinPoint, MaxMaxPoint);
+			LensInterpolationUtils::FIZMappingBilinearInterpolation<FDistortionMapPoint>(DesiredFocus, DesiredZoom, DistortionMapping, Result);
+
+			Test.TestEqual(*FString::Printf(TEXT("Find Indices - (%0.2f,%0.2f)"), DesiredFocus, DesiredZoom), bSuccess, true);
+			Test.TestEqual(*FString::Printf(TEXT("MinMin - (%0.2f,%0.2f)"), DesiredFocus, DesiredZoom), MinMinPoint, 85);
+			Test.TestEqual(*FString::Printf(TEXT("MinMax - (%0.2f,%0.2f)"), DesiredFocus, DesiredZoom), MinMaxPoint, 86);
+			Test.TestEqual(*FString::Printf(TEXT("MaxMin - (%0.2f,%0.2f)"), DesiredFocus, DesiredZoom), MaxMinPoint, 85);
+			Test.TestEqual(*FString::Printf(TEXT("MaxMax - (%0.2f,%0.2f)"), DesiredFocus, DesiredZoom), MaxMaxPoint, 86);
+			Test.TestEqual(*FString::Printf(TEXT("InterpResult - (%0.2f,%0.2f)"), DesiredFocus, DesiredZoom), Result.DistortionInfo.Parameters[0], 85.5f);
+		}
+
+		{
+			DesiredFocus = 110.0f;
+			DesiredZoom = -1.0f;
+			bSuccess = LensInterpolationUtils::FindBilinearInterpIndices<FDistortionMapPoint>(DesiredFocus, DesiredZoom, DistortionMapping, MinMinPoint, MinMaxPoint, MaxMinPoint, MaxMaxPoint);
+			LensInterpolationUtils::FIZMappingBilinearInterpolation<FDistortionMapPoint>(DesiredFocus, DesiredZoom, DistortionMapping, Result);
+
+			Test.TestEqual(*FString::Printf(TEXT("Find Indices - (%0.2f,%0.2f)"), DesiredFocus, DesiredZoom), bSuccess, true);
+			Test.TestEqual(*FString::Printf(TEXT("MinMin - (%0.2f,%0.2f)"), DesiredFocus, DesiredZoom), MinMinPoint, 82);
+			Test.TestEqual(*FString::Printf(TEXT("MinMax - (%0.2f,%0.2f)"), DesiredFocus, DesiredZoom), MinMaxPoint, 82);
+			Test.TestEqual(*FString::Printf(TEXT("MaxMin - (%0.2f,%0.2f)"), DesiredFocus, DesiredZoom), MaxMinPoint, 82);
+			Test.TestEqual(*FString::Printf(TEXT("MaxMax - (%0.2f,%0.2f)"), DesiredFocus, DesiredZoom), MaxMaxPoint, 82);
+			Test.TestEqual(*FString::Printf(TEXT("InterpResult - (%0.2f,%0.2f)"), DesiredFocus, DesiredZoom), Result.DistortionInfo.Parameters[0], 82.0f);
+		}
+
+		{
+			DesiredFocus = 35.0f;
+			DesiredZoom = -1.0f;
+			bSuccess = LensInterpolationUtils::FindBilinearInterpIndices<FDistortionMapPoint>(DesiredFocus, DesiredZoom, DistortionMapping, MinMinPoint, MinMaxPoint, MaxMinPoint, MaxMaxPoint);
+			LensInterpolationUtils::FIZMappingBilinearInterpolation<FDistortionMapPoint>(DesiredFocus, DesiredZoom, DistortionMapping, Result);
+
+			Test.TestEqual(*FString::Printf(TEXT("Find Indices - (%0.2f,%0.2f)"), DesiredFocus, DesiredZoom), bSuccess, true);
+			Test.TestEqual(*FString::Printf(TEXT("MinMin - (%0.2f,%0.2f)"), DesiredFocus, DesiredZoom), MinMinPoint, 30);
+			Test.TestEqual(*FString::Printf(TEXT("MinMax - (%0.2f,%0.2f)"), DesiredFocus, DesiredZoom), MinMaxPoint, 30);
+			Test.TestEqual(*FString::Printf(TEXT("MaxMin - (%0.2f,%0.2f)"), DesiredFocus, DesiredZoom), MaxMinPoint, 40);
+			Test.TestEqual(*FString::Printf(TEXT("MaxMax - (%0.2f,%0.2f)"), DesiredFocus, DesiredZoom), MaxMaxPoint, 40);
+			Test.TestEqual(*FString::Printf(TEXT("InterpResult - (%0.2f,%0.2f)"), DesiredFocus, DesiredZoom), Result.DistortionInfo.Parameters[0], 35.0f);
+		}
+	}
+}
+
+
+bool FTestCameraCalibrationCore::RunTest(const FString& Parameters)
+{
+	CameraCalibrationTestUtil::TestBilinearFindGrid(*this);
+
+    return true;
+}
+
+#endif // WITH_DEV_AUTOMATION_TESTS
+
diff --git a/Engine/Plugins/VirtualProduction/CameraCalibrationCore/Source/CameraCalibrationCore/Public/CalibrationPointComponent.h b/Engine/Plugins/VirtualProduction/CameraCalibrationCore/Source/CameraCalibrationCore/Public/CalibrationPointComponent.h
new file mode 100644
--- /dev/null
+++ b/Engine/Plugins/VirtualProduction/CameraCalibrationCore/Source/CameraCalibrationCore/Public/CalibrationPointComponent.h
@@ -0,0 +1,19 @@
+// Copyright Epic Games, Inc. All Rights Reserved.
+
+#pragma once
+
+#include "Components/SceneComponent.h"
+
+#include "CalibrationPointComponent.generated.h"
+
+/**
+ * One or more instances of this component can be added to an actor (e.g. a static mesh actor blueprint), 
+ * and should be placed at geometrically and visually distinct landmarks of the object.
+ * These 3d points will then be optionally used by any given nodal offset tool implementation to
+ * make a 3d-2d correspondence with the 2d points detected in the live action media.
+ */
+UCLASS(ClassGroup = (Calibration), meta = (BlueprintSpawnableComponent), meta = (DisplayName = "Calibration Point"))
+class CAMERACALIBRATIONCORE_API UCalibrationPointComponent : public USceneComponent
+{
+	GENERATED_BODY()
+};
\ No newline at end of file
diff --git a/Engine/Plugins/VirtualProduction/CameraCalibrationCore/Source/CameraCalibrationCore/Public/CameraCalibrationSettings.h b/Engine/Plugins/VirtualProduction/CameraCalibrationCore/Source/CameraCalibrationCore/Public/CameraCalibrationSettings.h
new file mode 100644
--- /dev/null
+++ b/Engine/Plugins/VirtualProduction/CameraCalibrationCore/Source/CameraCalibrationCore/Public/CameraCalibrationSettings.h
@@ -0,0 +1,127 @@
+// Copyright Epic Games, Inc. All Rights Reserved.
+
+#pragma once
+
+#include "Engine/DeveloperSettings.h"
+
+#include "LensDistortionModelHandlerBase.h"
+#include "LensFile.h"
+
+#include "CameraCalibrationSettings.generated.h"
+
+
+
+/**
+ * Settings for the CameraCalibration plugin modules. 
+ */
+UCLASS(config=Game)
+class CAMERACALIBRATIONCORE_API UCameraCalibrationSettings : public UDeveloperSettings
+{
+	GENERATED_BODY()
+
+public:
+	UCameraCalibrationSettings();
+
+	//~ Begin UDevelopperSettings interface
+	virtual FName GetCategoryName() const;
+#if WITH_EDITOR
+	virtual FText GetSectionText() const override;
+#endif
+	//~ End UDevelopperSettings interface
+
+public:
+
+	/**
+	 * Get the default startup lens file.
+	 *
+	 * @return The lens file, or nullptr if not set.
+	 */
+	ULensFile* GetStartupLensFile() const;
+
+	/** Get the default MaterialInterface used by the input Model Handler class to write the undistortion displacement map */
+	UMaterialInterface* GetDefaultUndistortionDisplacementMaterial(const TSubclassOf<ULensDistortionModelHandlerBase>& InModelHandler) const;
+
+	/** Get the default MaterialInterface used by the input Model Handler class to write the distortion displacement map */
+	UMaterialInterface* GetDefaultDistortionDisplacementMaterial(const TSubclassOf<ULensDistortionModelHandlerBase>& InModelHandler) const;
+
+	/** Get the default MaterialInterface used by the input Model Handler class to apply the post-process lens distortion effect */
+	UMaterialInterface* GetDefaultDistortionMaterial(const TSubclassOf<ULensDistortionModelHandlerBase>& InModelHandler) const;
+
+private:
+
+	/** 
+	 * Startup lens file for the project 
+	 * Can be overriden. Priority of operation is
+	 * 1. Apply startup lens file found in 'CameraCalibration.StartupLensFile' cvar at launch
+	 * 2. If none found, apply user startup file (only for editor runs)
+	 * 3. If none found, apply projet startup file (this one)
+	 */
+	UPROPERTY(config, EditAnywhere, Category = "Settings", meta = (ConfigRestartRequired = true))
+	TSoftObjectPtr<ULensFile> StartupLensFile;
+
+	/** Map of Lens Distortion Model Handler classes to the default displacement map material used by that class */
+	UPROPERTY(config)
+	TMap<TSubclassOf<ULensDistortionModelHandlerBase>, TSoftObjectPtr<UMaterialInterface>> DefaultUndistortionDisplacementMaterials;
+
+	/** Map of Lens Distortion Model Handler classes to the default displacement map material used by that class */
+	UPROPERTY(config)
+	TMap<TSubclassOf<ULensDistortionModelHandlerBase>, TSoftObjectPtr<UMaterialInterface>> DefaultDistortionDisplacementMaterials;
+
+	/** Map of Lens Distortion Model Handler classes to the default lens distortion post-process material used by that class */
+	UPROPERTY(config)
+	TMap<TSubclassOf<ULensDistortionModelHandlerBase>, TSoftObjectPtr<UMaterialInterface>> DefaultDistortionMaterials;
+};
+
+/**
+ * Settings for the camera calibration when in editor and standalone.
+ * @note Cooked games don't use this setting.
+ */
+UCLASS(config = EditorPerProjectUserSettings)
+class CAMERACALIBRATIONCORE_API UCameraCalibrationEditorSettings : public UDeveloperSettings
+{
+	GENERATED_BODY()
+
+public:
+
+		//~ Begin UDevelopperSettings interface
+		virtual FName GetCategoryName() const;
+#if WITH_EDITOR
+		virtual FText GetSectionText() const override;
+#endif
+		//~ End UDevelopperSettings interface
+
+public:
+
+#if WITH_EDITORONLY_DATA
+
+	/**
+	 * True if a lens file button shortcut should be added to level editor toolbar.
+	 */
+	UPROPERTY(config, EditAnywhere, Category = "Settings", Meta = (ConfigRestartRequired = true, DisplayName = "Enable Lens File Toolbar Button"))
+	bool bShowEditorToolbarButton = false;
+
+private:
+
+	/** 
+	 * Startup lens file per user in editor 
+	 * Can be overridden. Priority of operation is
+	 * 1. Apply startup lens file found in 'CameraCalibration.StartupLensFile' cvar at launch
+	 * 2. If none found, apply user startup file (this one)
+	 * 3. If none found, apply project startup file
+	 */
+	UPROPERTY(config, EditAnywhere, Category = "Settings")
+	TSoftObjectPtr<ULensFile> UserLensFile;
+#endif
+
+public:
+
+	/**
+	 * Get the lens file used by the engine when in the editor and standalone.
+	 *
+	 * @return The lens file, or nullptr if not set.
+	 */
+	ULensFile* GetUserLensFile() const;
+
+	/** Set the lens file used by the engine when in the editor and standalone. */
+	void SetUserLensFile(ULensFile* InLensFile);
+};
\ No newline at end of file
diff --git a/Engine/Plugins/VirtualProduction/CameraCalibrationCore/Source/CameraCalibrationCore/Public/CameraCalibrationStep.h b/Engine/Plugins/VirtualProduction/CameraCalibrationCore/Source/CameraCalibrationCore/Public/CameraCalibrationStep.h
new file mode 100644
--- /dev/null
+++ b/Engine/Plugins/VirtualProduction/CameraCalibrationCore/Source/CameraCalibrationCore/Public/CameraCalibrationStep.h
@@ -0,0 +1,53 @@
+// Copyright Epic Games, Inc. All Rights Reserved.
+
+#pragma once
+
+#include "CoreMinimal.h"
+#include "Templates/SharedPointer.h"
+#include "Widgets/SBoxPanel.h"
+
+#include "CameraCalibrationStep.generated.h"
+
+struct FGeometry;
+struct FPointerEvent;
+
+class FCameraCalibrationStepsController;
+
+/**
+ * Interface of a camera calibration step. These will appear in a Camera Calibration Toolkit tab.
+ */
+UCLASS(Abstract)
+class CAMERACALIBRATIONCORE_API UCameraCalibrationStep : public UObject
+{
+	GENERATED_BODY()
+
+public:
+
+	/** Make sure you initialize before using the object */
+	virtual void Initialize(FCameraCalibrationStepsController* InCameraCalibrationStepController) {};
+
+	/** Clean up resources and don't use CameraCalibrationStepController anymore */
+	virtual void Shutdown() {};
+
+	/** Called every frame */
+	virtual void Tick(float DeltaTime) {};
+
+	/** Callback when viewport is clicked. Returns false if the event was not handled. */
+	virtual bool OnViewportClicked(const FGeometry& MyGeometry, const FPointerEvent& MouseEvent) { return false;  };
+
+	/** Returns the UI of this camera calibration step. Expected to only be called once */
+	virtual TSharedRef<SWidget> BuildUI() { return SNew(SVerticalBox); };
+
+	/** Returns a title or friendly name that can be placed in UI */
+	virtual FName FriendlyName() const { return TEXT("Invalid"); };
+
+	/** Returns true if the given calibration step is a known prerequisite for this step */
+	virtual bool DependsOnStep(UCameraCalibrationStep* Step) const { return false; };
+
+	/** Called when this step is the active step in the UI */
+	virtual void Activate() {};
+
+	/** Called when this step is no longer the active step in the UI */
+	virtual void Deactivate() {};
+
+};
diff --git a/Engine/Plugins/VirtualProduction/CameraCalibrationCore/Source/CameraCalibrationCore/Public/CameraCalibrationSubsystem.h b/Engine/Plugins/VirtualProduction/CameraCalibrationCore/Source/CameraCalibrationCore/Public/CameraCalibrationSubsystem.h
new file mode 100644
--- /dev/null
+++ b/Engine/Plugins/VirtualProduction/CameraCalibrationCore/Source/CameraCalibrationCore/Public/CameraCalibrationSubsystem.h
@@ -0,0 +1,121 @@
+// Copyright Epic Games, Inc. All Rights Reserved.
+
+#pragma once
+
+#include "Subsystems/EngineSubsystem.h"
+
+#include "CoreMinimal.h"
+
+#include "CameraCalibrationTypes.h"
+#include "CameraNodalOffsetAlgo.h"
+#include "CameraCalibrationStep.h"
+#include "Containers/ArrayView.h"
+#include "LensDistortionModelHandlerBase.h"
+#include "LensFile.h"
+#include "UObject/ObjectKey.h"
+
+#include "CameraCalibrationSubsystem.generated.h"
+
+class UCineCameraComponent;
+
+/**
+ * Camera Calibration subsystem
+ */
+UCLASS()
+class CAMERACALIBRATIONCORE_API UCameraCalibrationSubsystem : public UEngineSubsystem
+{
+	GENERATED_BODY()
+
+public:
+
+	/** Get the default lens file. */
+	UFUNCTION(BlueprintCallable, Category = "Lens Distortion")
+	ULensFile* GetDefaultLensFile() const;
+
+	/** Get the default lens file. */
+	UFUNCTION(BlueprintCallable, Category = "Lens Distortion")
+	void SetDefaultLensFile(ULensFile* NewDefaultLensFile);
+
+	/** Facilitator around the picker to get the desired lens file. */
+	UFUNCTION(BlueprintCallable, Category = "Lens Distortion")
+	ULensFile* GetLensFile(const FLensFilePicker& Picker) const;
+
+	/** Return all handlers associated with the input camera component */
+	UFUNCTION(BlueprintCallable, Category = "Lens Distortion")
+	TArray<ULensDistortionModelHandlerBase*> GetDistortionModelHandlers(UCineCameraComponent* Component);
+
+	/** 
+	 * Return the handler associated with the input distortion source, if one exists 
+	 * If bUpdatePicker is true, the input picker reference will be updated so that its properties match those of the found handler
+	 */
+	UFUNCTION(BlueprintCallable, Category = "Lens Distortion")
+	ULensDistortionModelHandlerBase* FindDistortionModelHandler(FDistortionHandlerPicker& DistortionHandlerPicker, bool bUpdatePicker = true) const;
+
+	/** Return the handler associated with the input distortion source, if one exists that also matches the input model. If none exist, create a new handler and return it. */
+	UFUNCTION(BlueprintCallable, Category = "Lens Distortion")
+	ULensDistortionModelHandlerBase* FindOrCreateDistortionModelHandler(FDistortionHandlerPicker& DistortionHandlerPicker, const TSubclassOf<ULensModel> LensModelClass);
+
+	/** Disassociate the input handler from the input camera component in the subsystem's handler map */
+	UFUNCTION(BlueprintCallable, Category = "Lens Distortion")
+	void UnregisterDistortionModelHandler(UCineCameraComponent* Component, ULensDistortionModelHandlerBase* Handler);
+
+	/** Return the ULensModel subclass that was registered with the input model name */
+	UFUNCTION(BlueprintCallable, Category = "Lens Distortion")
+	TSubclassOf<ULensModel> GetRegisteredLensModel(FName ModelName) const;
+
+	/** Returns the nodal offset algorithm by name */
+	UFUNCTION(BlueprintCallable, Category = "Lens Distortion")
+	TSubclassOf<UCameraNodalOffsetAlgo> GetCameraNodalOffsetAlgo(FName Name) const;
+
+	/** Returns an array with the names of the available nodal offset algorithms */
+	UFUNCTION(BlueprintCallable, Category = "Lens Distortion")
+	TArray<FName> GetCameraNodalOffsetAlgos() const;
+
+	/** Returns an array with the names of the available camera calibration steps */
+	UFUNCTION(BlueprintCallable, Category = "Lens Distortion")
+	TArray<FName> GetCameraCalibrationSteps() const;
+
+	/** Returns the camera calibration step by name */
+	UFUNCTION(BlueprintCallable, Category = "Lens Distortion")
+	TSubclassOf<UCameraCalibrationStep> GetCameraCalibrationStep(FName Name) const;
+
+public:
+	//~ Begin USubsystem interface
+	virtual void Deinitialize() override;
+	virtual void Initialize(FSubsystemCollectionBase& Collection) override;
+	//~ End USubsystem interface
+
+	/** Add a new Lens Model to the registered model map */
+	void RegisterDistortionModel(TSubclassOf<ULensModel> LensModel);
+
+	/** Remove a Lens Model from the registered model map */
+	void UnregisterDistortionModel(TSubclassOf<ULensModel> LensModel);
+
+private:
+	/** Default lens file to use when no override has been provided */
+	UPROPERTY(Transient)
+	ULensFile* DefaultLensFile = nullptr;
+
+	/** Map of model names to ULensModel subclasses */
+	UPROPERTY(Transient)
+	TMap<FName, TSubclassOf<ULensModel>> LensModelMap;
+
+	/** Holds the registered camera nodal offset algos */
+	UPROPERTY(Transient)
+	TMap<FName, TSubclassOf<UCameraNodalOffsetAlgo>> CameraNodalOffsetAlgosMap;
+
+	/** Holds the registered camera calibration steps */
+	UPROPERTY(Transient)
+	TMap<FName, TSubclassOf<UCameraCalibrationStep>> CameraCalibrationStepsMap;
+
+	/** Map of actor components to the authoritative lens model that should be used with that component */
+	TMap<FObjectKey, TSubclassOf<ULensModel>> ComponentsWithAuthoritativeModels;
+
+	/** Map of actor components to the authoritative lens model that should be used with that component */
+	TMultiMap<FObjectKey, ULensDistortionModelHandlerBase*> LensDistortionHandlerMap;
+
+private:
+
+	FDelegateHandle PostEngineInitHandle;
+
+};
diff --git a/Engine/Plugins/VirtualProduction/CameraCalibrationCore/Source/CameraCalibrationCore/Public/CameraCalibrationTypes.h b/Engine/Plugins/VirtualProduction/CameraCalibrationCore/Source/CameraCalibrationCore/Public/CameraCalibrationTypes.h
new file mode 100644
--- /dev/null
+++ b/Engine/Plugins/VirtualProduction/CameraCalibrationCore/Source/CameraCalibrationCore/Public/CameraCalibrationTypes.h
@@ -0,0 +1,29 @@
+// Copyright Epic Games, Inc. All Rights Reserved.
+
+#pragma once
+
+#include "CineCameraComponent.h"
+#include "LensDistortionModelHandlerBase.h"
+#include "Misc/Guid.h"
+
+#include "CameraCalibrationTypes.generated.h"
+
+/** Utility structure for selecting a distortion handler from the camera calibration subsystem */
+USTRUCT(BlueprintType)
+struct CAMERACALIBRATIONCORE_API FDistortionHandlerPicker
+{
+	GENERATED_BODY()
+
+public:
+	/** CineCameraComponent with which the desired distortion handler is associated */
+	UPROPERTY(Transient)
+	UCineCameraComponent* TargetCameraComponent = nullptr;
+
+	/** UObject that produces the distortion state for the desired distortion handler */
+	UPROPERTY()
+	FGuid DistortionProducerID;
+
+	/** Display name of the desired distortion handler */
+	UPROPERTY()
+	FString HandlerDisplayName;
+};
\ No newline at end of file
diff --git a/Engine/Plugins/VirtualProduction/CameraCalibrationCore/Source/CameraCalibrationCore/Public/CameraCalibrationUtils.h b/Engine/Plugins/VirtualProduction/CameraCalibrationCore/Source/CameraCalibrationCore/Public/CameraCalibrationUtils.h
new file mode 100644
--- /dev/null
+++ b/Engine/Plugins/VirtualProduction/CameraCalibrationCore/Source/CameraCalibrationCore/Public/CameraCalibrationUtils.h
@@ -0,0 +1,33 @@
+// Copyright Epic Games, Inc. All Rights Reserved.
+
+#pragma once
+
+#include "CoreMinimal.h"
+
+struct FTransform;
+
+/**
+ * Helper class for commonly used functions for camera calibration.
+ */
+class CAMERACALIBRATIONCORE_API FCameraCalibrationUtils
+{
+public:
+
+	/** Enumeration to specify any cartesian axis in positive or negative directions */
+	enum class EAxis
+	{
+		X, Y, Z,
+		Xn, Yn, Zn,
+	};
+
+public:
+
+	/** Converts in-place the coordinate system of the given FTransform by specifying the source axes in terms of the destionation axes */
+	static void ConvertCoordinateSystem(FTransform& Transform, EAxis DstXInSrcAxis, EAxis DstYInSrcAxis, EAxis DstZInSrcAxis);
+
+	/** Converts in-place an FTransform in Unreal coordinates to OpenCV coordinates */
+	static void ConvertUnrealToOpenCV(FTransform& Transform);
+
+	/** Converts in-place an FTransform in OpenCV coordinates to Unreal coordinates */
+	static void ConvertOpenCVToUnreal(FTransform& Transform);
+};
diff --git a/Engine/Plugins/VirtualProduction/CameraCalibrationCore/Source/CameraCalibrationCore/Public/CameraNodalOffsetAlgo.h b/Engine/Plugins/VirtualProduction/CameraCalibrationCore/Source/CameraCalibrationCore/Public/CameraNodalOffsetAlgo.h
new file mode 100644
--- /dev/null
+++ b/Engine/Plugins/VirtualProduction/CameraCalibrationCore/Source/CameraCalibrationCore/Public/CameraNodalOffsetAlgo.h
@@ -0,0 +1,66 @@
+// Copyright Epic Games, Inc. All Rights Reserved.
+
+#pragma once
+
+#include "CoreMinimal.h"
+#include "Widgets/SBoxPanel.h"
+
+#include "CameraNodalOffsetAlgo.generated.h"
+
+struct FGeometry;
+struct FNodalPointOffset;
+struct FPointerEvent;
+struct FTransform;
+
+class FNodalOffsetTool;
+class SWidget;
+
+/**
+ * UCameraNodalOffsetAlgo defines the interface that any nodal calibration point algorithm should implement
+ * in order to be used and listed by the Nodal Offset Tool. 
+ */
+UCLASS(Abstract)
+class CAMERACALIBRATIONCORE_API UCameraNodalOffsetAlgo : public UObject
+{
+	GENERATED_BODY()
+
+public:
+
+	/** Structure with the camera information needed by the algorithm */
+	struct FCameraData
+	{
+		FTransform Pose;
+		uint32 CameraId;
+		float Zoom;
+		float Focus;
+		float FilmbackX;
+		float FilmbackY;
+		bool bValid;
+	};
+
+public:
+
+	/** Make sure you initialize before using the object */
+	virtual void Initialize(TWeakPtr<FNodalOffsetTool> InNodalOffsetTool) {};
+
+	/** Clean up resources and don't use NodalOffsetTool anymore */
+	virtual void Shutdown() {};
+
+	/** Called every frame */
+	virtual void Tick(float DeltaTime) {};
+
+	/** Callback when viewport is clicked. Returns false if the event was not handled. */
+	virtual bool OnViewportClicked(const FGeometry& MyGeometry, const FPointerEvent& MouseEvent) { return false;  };
+
+	/** Returns the UI of this calibrator. Expected to only be called once */
+	virtual TSharedRef<SWidget> BuildUI() { return SNew(SVerticalBox); };
+
+	/** Returns the most recently calibrated nodal offset transform, with an error metric */
+	virtual bool GetNodalOffset(FNodalPointOffset& OutNodalOffset, float& OutFocus, float& OutZoom, float& OutError, FText& OutErrorMessage) { return false; };
+
+	/** Returns a descriptive name/title of this nodal offset algorithm */
+	virtual FName FriendlyName() const { return TEXT("Invalid Name"); };
+
+	/** Called when the current offset was saved */
+	virtual void OnSavedNodalOffset() { };
+};
diff --git a/Engine/Plugins/VirtualProduction/CameraCalibrationCore/Source/CameraCalibrationCore/Public/ICalibratedMapProcessor.h b/Engine/Plugins/VirtualProduction/CameraCalibrationCore/Source/CameraCalibrationCore/Public/ICalibratedMapProcessor.h
new file mode 100644
--- /dev/null
+++ b/Engine/Plugins/VirtualProduction/CameraCalibrationCore/Source/CameraCalibrationCore/Public/ICalibratedMapProcessor.h
@@ -0,0 +1,28 @@
+// Copyright Epic Games, Inc. All Rights Reserved.
+
+#pragma once
+
+#include "CoreMinimal.h"
+
+struct FDerivedDistortionDataJobOutput;
+struct FDerivedDistortionDataJobArgs;
+
+/** Delegate triggered when derived data job has completed, successfully or not. */
+DECLARE_DELEGATE_OneParam(FOnDerivedDistortionJobCompleted, const FDerivedDistortionDataJobOutput& /*JobOutput*/);
+
+
+/**
+ * Interface to handle stmaps (calibrated maps) and derive distortion data out of it
+ */
+class ICalibratedMapProcessor
+{
+public:
+
+	virtual ~ICalibratedMapProcessor() = default;
+
+	/** Updates processor to manage job queues */
+	virtual void Update() = 0;
+
+	/** Push new job to derive distortion data from calibrated map */
+	virtual bool PushDerivedDistortionDataJob(FDerivedDistortionDataJobArgs&& JobArgs) = 0;
+};
\ No newline at end of file
diff --git a/Engine/Plugins/VirtualProduction/CameraCalibrationCore/Source/CameraCalibrationCore/Public/LensData.h b/Engine/Plugins/VirtualProduction/CameraCalibrationCore/Source/CameraCalibrationCore/Public/LensData.h
new file mode 100644
--- /dev/null
+++ b/Engine/Plugins/VirtualProduction/CameraCalibrationCore/Source/CameraCalibrationCore/Public/LensData.h
@@ -0,0 +1,91 @@
+// Copyright Epic Games, Inc. All Rights Reserved.
+
+#pragma once
+
+#include "CoreTypes.h"
+
+#include "Engine/EngineTypes.h"
+#include "Models/LensModel.h"
+#include "UObject/ObjectMacros.h"
+#include "UObject/Object.h"
+
+#include "LensData.generated.h"
+
+
+/**
+ * Information about the lens rig
+ */
+USTRUCT(BlueprintType)
+struct CAMERACALIBRATIONCORE_API FLensInfo
+{
+	GENERATED_BODY()
+
+public:
+
+	/** Model name of the lens */
+	UPROPERTY(EditAnywhere, Category = "Lens Info")
+	FString LensModelName;
+
+	/** Serial number of the lens */
+	UPROPERTY(EditAnywhere, Category = "Lens Info")
+	FString LensSerialNumber;
+
+	/** Model of the lens (spherical, anamorphic, etc...) */
+	UPROPERTY(EditAnywhere, BlueprintReadOnly, Category = "Lens Info")
+	TSubclassOf<ULensModel> LensModel;
+
+	/** Width and height of the calibrated camera's sensor, in millimeters */
+	UPROPERTY(EditAnywhere, Category = "Lens Info")
+	FVector2D SensorDimensions = FVector2D(23.76f, 13.365f);
+};
+
+/**
+ * Lens distortion parameters
+ */
+USTRUCT(BlueprintType)
+struct CAMERACALIBRATIONCORE_API FDistortionInfo
+{
+	GENERATED_BODY()
+
+public:
+	/** Generic array of floating-point lens distortion parameters */
+	UPROPERTY(EditAnywhere, BlueprintReadWrite, Category = "Distortion")
+	TArray<float> Parameters;
+
+	/** Value expected to be normalized (unitless) */
+	UPROPERTY(EditAnywhere, Category = "Camera")
+	FVector2D FxFy = FVector2D(1.0f, (16.0f / 9.0f));
+};
+
+/**
+ * Lens camera parameters
+ */
+USTRUCT(BlueprintType)
+struct CAMERACALIBRATIONCORE_API FIntrinsicParameters
+{
+	GENERATED_BODY()
+
+public:
+	/** Value expected to be normalized [0,1] */
+	UPROPERTY(EditAnywhere, Category = "Camera", meta = (DisplayName = "Image Center"))
+	FVector2D PrincipalPoint = FVector2D(0.5f, 0.5f);
+};
+
+/**
+ * Lens nodal point offset
+ */
+USTRUCT(BlueprintType)
+struct CAMERACALIBRATIONCORE_API FNodalPointOffset
+{
+	GENERATED_BODY()
+
+public:
+
+	UPROPERTY(EditAnywhere, Category = "Nodal point")
+	FVector LocationOffset = FVector::ZeroVector;
+
+	UPROPERTY(EditAnywhere, Category = "Nodal point")
+	FQuat RotationOffset = FQuat::Identity;
+};
+
+
diff --git a/Engine/Plugins/VirtualProduction/CameraCalibrationCore/Source/CameraCalibrationCore/Public/LensDistortionModelHandlerBase.h b/Engine/Plugins/VirtualProduction/CameraCalibrationCore/Source/CameraCalibrationCore/Public/LensDistortionModelHandlerBase.h
new file mode 100644
--- /dev/null
+++ b/Engine/Plugins/VirtualProduction/CameraCalibrationCore/Source/CameraCalibrationCore/Public/LensDistortionModelHandlerBase.h
@@ -0,0 +1,175 @@
+// Copyright Epic Games, Inc. All Rights Reserved.
+
+#pragma once
+
+#include "CoreMinimal.h"
+#include "LensData.h"
+#include "Materials/MaterialInstanceDynamic.h"
+#include "Misc/Guid.h"
+#include "Models/LensModel.h"
+
+#include "LensDistortionModelHandlerBase.generated.h"
+
+USTRUCT(BlueprintType)
+struct CAMERACALIBRATIONCORE_API FLensDistortionState
+{
+	GENERATED_BODY()
+
+public:
+	/** Generic array of distortion parameters */
+	UPROPERTY(EditAnywhere, BlueprintReadWrite, Category = "Distortion")
+	FDistortionInfo DistortionInfo;
+
+	/** Normalized center of the image, in the range [0.0f, 1.0f] */
+	UPROPERTY(EditAnywhere, BlueprintReadWrite, Category = "Distortion", meta = (DisplayName = "Image Center"))
+	FVector2D PrincipalPoint = FVector2D(0.5f, 0.5f);
+
+public:
+	bool operator==(const FLensDistortionState& Other) const;
+	bool operator!=(const FLensDistortionState& Other) const { return !(*this == Other); }
+};
+
+/** Asset user data that can be used on Camera Actors to manage lens distortion state and utilities  */
+UCLASS(Abstract)
+class CAMERACALIBRATIONCORE_API ULensDistortionModelHandlerBase : public UObject
+{
+	GENERATED_BODY()
+
+public:
+	/** Returns true if the input model is supported by this model handler, false otherwise. */
+	UFUNCTION(BlueprintCallable, Category = "Distortion")
+	bool IsModelSupported(const TSubclassOf<ULensModel>& ModelToSupport) const;
+
+	/** Update the lens distortion state, recompute the overscan factor, and set all material parameters */
+	UFUNCTION(BlueprintCallable, Category = "Distortion")
+    void SetDistortionState(const FLensDistortionState& InNewState);
+
+	/** Get the UV displacement map used to undistort a distorted image */
+	UFUNCTION(BlueprintCallable, Category = "Distortion")
+	UTextureRenderTarget2D* GetUndistortionDisplacementMap() const { return UndistortionDisplacementMapRT; }
+
+	/** Get the UV displacement map used to distort an undistorted image */
+	UFUNCTION(BlueprintCallable, Category = "Distortion")
+	UTextureRenderTarget2D* GetDistortionDisplacementMap() const { return DistortionDisplacementMapRT; }
+
+public:
+	//~ Begin UObject Interface
+	virtual void PostInitProperties() override;
+
+#if WITH_EDITOR
+	virtual void PostEditChangeChainProperty(struct FPropertyChangedChainEvent& PropertyChangedEvent) override;
+#endif	
+	//~ End UObject Interface
+
+	/** Get the current distortion state (the lens model and properties that mathematically represent the distortion characteristics */
+	FLensDistortionState GetCurrentDistortionState() const { return CurrentState; }
+
+	/** Get the post-process MID for the currently specified lens model */
+	UMaterialInstanceDynamic* GetDistortionMID() const { return DistortionPostProcessMID; }
+
+	/** Get the specified lens model that characterizes the distortion effect */
+	const TSubclassOf<ULensModel>& GetLensModelClass() const { return LensModelClass; };
+
+	/** Get the UObject that produces the distortion state for this handler */
+	FGuid GetDistortionProducerID() const { return DistortionProducerID; }
+
+	/** Set the UObject that produces the distortion state for this handler */
+	void SetDistortionProducerID(const FGuid& InDistortionProducerID) { DistortionProducerID = InDistortionProducerID; }
+
+	/** Get the display name of this lens distortion model handler */
+	FString GetDisplayName() const { return DisplayName; }
+
+	/** Set the display name of this lens distortion model handler */
+	void SetDisplayName(FString InDisplayName) { DisplayName = InDisplayName; }
+
+	/** Get the normalized center of projection of the image, in the range [0.0f, 1.0f] */
+	FVector2D GetPrincipalPoint() const { return CurrentState.PrincipalPoint; }
+
+	/** Get the normalized focal length (unitless) */
+	FVector2D GetFxFy() const { return CurrentState.DistortionInfo.FxFy; }
+
+	/** Updates overscan factor and applies to material instances */
+	void SetOverscanFactor(float OverscanFactor);
+
+	/** Returns the last overscan factor that was set */
+	float GetOverscanFactor() const { return OverscanFactor; }
+
+	/** Use the current distortion state to compute the overscan factor needed such that all distorted UVs will fall into the valid range of [0,1] */
+	float ComputeOverscanFactor() const;
+
+	/** Computes the distorted version of UndistortedUVs based on the current state */
+	TArray<FVector2D> GetDistortedUVs(TConstArrayView<FVector2D> UndistortedUVs) const;
+
+	/** Draw the undistortion displacement map associated with the current state to the DestinationTexture */
+	bool DrawUndistortionDisplacementMap(UTextureRenderTarget2D* DestinationTexture);
+
+	/** Draw the distortion displacement map associated with the current state to the DestinationTexture */
+	bool DrawDistortionDisplacementMap(UTextureRenderTarget2D* DestinationTexture);
+
+	/** Draws the current distortion state to the internal displacement map */
+	void ProcessCurrentDistortion();
+
+protected:
+	/** Initialize the handler. Derived classes must set the LensModelClass that they support, if not already set */
+	virtual void InitializeHandler() PURE_VIRTUAL(ULensDistortionModelHandlerBase::InitializeHandler);
+
+	/** Use the current distortion state to compute the distortion position of an input UV coordinate */
+	virtual FVector2D ComputeDistortedUV(const FVector2D& InScreenUV) const PURE_VIRTUAL(ULensDistortionModelHandlerBase::ComputeDistortedUV, return FVector2D::ZeroVector;);
+
+	/** Create the distortion MIDs */
+	virtual void InitDistortionMaterials() PURE_VIRTUAL(ULensDistortionModelHandlerBase::InitDistortionMaterials);
+
+	/** Set the material parameters for the displacement map and distortion post-process materials */
+	virtual void UpdateMaterialParameters() PURE_VIRTUAL(ULensDistortionModelHandlerBase::UpdateMaterialParameters);
+
+	/** Convert the generic distortion parameter array into the specific structure of parameters used by the supported lens model */
+	virtual void InterpretDistortionParameters() PURE_VIRTUAL(ULensDistortionModelHandlerBase::InterpretDistortionParameters);
+
+protected:
+	/** Lens Model describing how to interpret the distortion parameters */
+	UPROPERTY(VisibleAnywhere, BlueprintReadWrite, Category = "Distortion")
+	TSubclassOf<ULensModel> LensModelClass;
+
+	/** Dynamically created post-process material instance for the currently specified lens model */
+	UPROPERTY(VisibleAnywhere, BlueprintReadOnly, Category = "Distortion")
+	UMaterialInstanceDynamic* DistortionPostProcessMID = nullptr;
+
+	/** Current state as set by the most recent call to Update() */
+	UPROPERTY(EditAnywhere, BlueprintReadOnly, Category = "Distortion", meta = (ShowOnlyInnerProperties))
+	FLensDistortionState CurrentState;
+
+	/** Display name, used to identify handler in-editor details panels */
+	UPROPERTY(VisibleAnywhere, Category = "Distortion")
+	FString DisplayName;
+
+	/** Computed overscan factor needed to scale the camera's FOV (read-only) */
+	UPROPERTY(VisibleAnywhere, BlueprintReadOnly, Category = "Distortion")
+	float OverscanFactor = 1.0f;
+
+	/** MID used to draw the undistortion displacement map */
+	UPROPERTY(Transient)
+	UMaterialInstanceDynamic* UndistortionDisplacementMapMID = nullptr;
+
+	/** MID used to draw the distortion displacement map */
+	UPROPERTY(Transient)
+	UMaterialInstanceDynamic* DistortionDisplacementMapMID = nullptr;
+
+	/** UV displacement map used to undistort a distorted image */
+	UPROPERTY(Transient)
+	UTextureRenderTarget2D* UndistortionDisplacementMapRT = nullptr;
+
+	/** UV displacement map used to distort an undistorted image */
+	UPROPERTY(Transient)
+	UTextureRenderTarget2D* DistortionDisplacementMapRT = nullptr;
+
+	/** UObject that is producing the distortion state for this handler */
+	UPROPERTY()
+	FGuid DistortionProducerID;
+
+private:
+	static constexpr uint32 DisplacementMapWidth = 256;
+	static constexpr uint32 DisplacementMapHeight = 256;
+
+	/** Tracks whether distortion state has been changed */
+	bool bIsDirty = true;
+};
diff --git a/Engine/Plugins/VirtualProduction/CameraCalibrationCore/Source/CameraCalibrationCore/Public/LensFile.h b/Engine/Plugins/VirtualProduction/CameraCalibrationCore/Source/CameraCalibrationCore/Public/LensFile.h
new file mode 100644
--- /dev/null
+++ b/Engine/Plugins/VirtualProduction/CameraCalibrationCore/Source/CameraCalibrationCore/Public/LensFile.h
@@ -0,0 +1,457 @@
+// Copyright Epic Games, Inc. All Rights Reserved.
+
+#pragma once
+
+#include "CoreTypes.h"
+
+#include "Curves/RichCurve.h"
+#include "Engine/Texture.h"
+#include "ICalibratedMapProcessor.h"
+#include "LensData.h"
+#include "Misc/Optional.h"
+#include "Templates/UniquePtr.h"
+#include "Tickable.h"
+#include "UObject/Object.h"
+#include "UObject/ObjectMacros.h"
+
+#include "LensFile.generated.h"
+
+
+class FLensFilePreComputeDataProcessor;
+class UCineCameraComponent;
+class ULensDistortionModelHandlerBase;
+
+
+/** Mode of operation of Lens File */
+UENUM()
+enum class ELensDataMode : uint8
+{
+	Parameters = 0,
+	STMap = 1
+};
+
+/**
+ * Distortion data evaluated for given FZ pair based on lens parameters
+ */
+USTRUCT(BlueprintType)
+struct CAMERACALIBRATIONCORE_API FDistortionData
+{
+	GENERATED_BODY()
+
+public:
+
+	UPROPERTY(VisibleAnywhere, Category = "Distortion")
+	TArray<FVector2D> DistortedUVs;
+
+	/** Estimated overscan factor based on distortion to have distorted cg covering full size */
+	UPROPERTY(EditAnywhere, Category = "Distortion")
+	float OverscanFactor = 1.0f;
+};
+
+/**
+ * Encoder mapping
+ */
+USTRUCT(BlueprintType)
+struct CAMERACALIBRATIONCORE_API FEncoderMapping
+{
+	GENERATED_BODY()
+
+public:
+
+	/** Focus curve from encoder values to nominal values */
+	UPROPERTY()
+	FRichCurve Focus;
+
+	/** Iris curve from encoder values to nominal values */
+	UPROPERTY()
+	FRichCurve Iris;
+
+	/** 
+	 * Zoom curve from encoder values to nominal values 
+	 * @note To be removed and use only Fx/Fy calibrated values
+	 */
+	UPROPERTY()
+	FRichCurve Zoom;
+};
+
+/**
+ * Derived data computed from parameters or stmap
+ */
+USTRUCT(BlueprintType)
+struct CAMERACALIBRATIONCORE_API FDerivedDistortionData
+{
+	GENERATED_BODY()
+
+	/** Precomputed data about distortion */
+	UPROPERTY(VisibleAnywhere, Category = "Distortion")
+	FDistortionData DistortionData;
+
+	/** Computed displacement map based on undistortion data */
+	UPROPERTY(Transient, VisibleAnywhere, Category = "Distortion")
+	UTextureRenderTarget2D* UndistortionDisplacementMap = nullptr;
+
+	/** Computed displacement map based on distortion data */
+	UPROPERTY(Transient, VisibleAnywhere, Category = "Distortion")
+	UTextureRenderTarget2D* DistortionDisplacementMap = nullptr;
+
+	/** When dirty, derived data needs to be recomputed */
+	bool bIsDirty = true;
+};
+
+/**
+ * A data point associating focus and zoom to lens parameters
+ */
+USTRUCT(BlueprintType)
+struct CAMERACALIBRATIONCORE_API FDistortionMapPoint
+{
+	GENERATED_BODY()
+
+public:
+	FDistortionMapPoint()
+    : Identifier(FGuid::NewGuid())
+	{}
+
+	const FGuid& GetIdentifier() const { return Identifier; }
+
+	/** Returns whether this point is considered valid */
+	bool IsValid() const { return true; }
+
+
+public:
+
+	UPROPERTY(EditAnywhere, Category = "Distortion")
+	float Focus = 0.0f;
+
+	UPROPERTY(EditAnywhere, Category = "Distortion")
+	float Zoom = 0.0f;
+
+	UPROPERTY(EditAnywhere, Category = "Distortion")
+	FDistortionInfo DistortionInfo;
+
+private:
+
+	/** Unique identifier for this map point to associate it with derived data */
+	FGuid Identifier;
+};
+
+
+/**
+* A data point associating focus and zoom to precalibrated STMap
+*/
+USTRUCT(BlueprintType)
+struct CAMERACALIBRATIONCORE_API FCalibratedMapPoint
+{
+	GENERATED_BODY()
+
+public:
+
+	FCalibratedMapPoint()
+		: Identifier(FGuid::NewGuid())
+	{}
+
+	/** Returns the identifier of this point */
+	const FGuid& GetIdentifier() const { return Identifier; }
+
+	/** Returns whether this point is considered valid */
+	bool IsValid() const 
+	{ 
+		return (DistortionMap != nullptr 
+			&& DerivedDistortionData.UndistortionDisplacementMap != nullptr
+			&& DerivedDistortionData.DistortionDisplacementMap != nullptr
+			&& DerivedDistortionData.bIsDirty == false);
+	}
+
+public:
+	UPROPERTY(EditAnywhere, Category = "Distortion")
+	float Focus = 0.0f;
+
+	UPROPERTY(EditAnywhere, Category = "Distortion")
+	float Zoom = 0.0f;
+
+	/** Pre calibrated UVMap/STMap
+	 * RG channels are expected to have undistortion map (from distorted to undistorted)
+	 * BA channels are expected to have distortion map (from undistorted (CG) to distorted)
+	 */
+	UPROPERTY(EditAnywhere, Category = "Distortion")
+	UTexture* DistortionMap = nullptr;
+
+	/** Derived distortion data associated with this point */
+	UPROPERTY(Transient)
+	FDerivedDistortionData DerivedDistortionData;
+
+private:
+
+	/** Unique identifier for this map point to associate it with derived data */
+	FGuid Identifier;
+};
+
+/**
+ * A data point associating focus and zoom to the principal point (image center)
+ */
+USTRUCT(BlueprintType)
+struct CAMERACALIBRATIONCORE_API FIntrinsicMapPoint
+{
+	GENERATED_BODY()
+
+public:
+
+	/** Returns whether this point is considered valid */
+	bool IsValid() const { return true; }
+
+public:
+
+	UPROPERTY(EditAnywhere, Category = "Center shift")
+	float Focus = 0.0f;
+
+	UPROPERTY(EditAnywhere, Category = "Center shift")
+	float Zoom = 0.0f;
+
+	UPROPERTY(EditAnywhere, Category = "Center shift")
+	FIntrinsicParameters Parameters;
+};
+
+/**
+ * A data point associating focus and zoom to Nodal offset
+ */
+USTRUCT(BlueprintType)
+struct CAMERACALIBRATIONCORE_API FNodalOffsetMapPoint
+{
+	GENERATED_BODY()
+
+public:
+
+	/** Returns whether this point is considered valid */
+	bool IsValid() const { return true; }
+
+public:
+
+	UPROPERTY(EditAnywhere, Category = "Nodal Point")
+	float Focus = 0.0f;
+
+	UPROPERTY(EditAnywhere, Category = "Nodal Point")
+	float Zoom = 0.0f;
+
+	UPROPERTY(EditAnywhere, Category = "Nodal Point")
+	FNodalPointOffset NodalOffset;
+};
+
+
+
+/**
+ * A Lens file containing calibration mapping from FIZ data
+ */
+UCLASS(BlueprintType)
+class CAMERACALIBRATIONCORE_API ULensFile : public UObject, public FTickableGameObject
+{
+	GENERATED_BODY()
+
+
+public:
+
+	ULensFile();
+
+	//~Begin UObject interface
+#if WITH_EDITOR
+	virtual void PostEditChangeChainProperty( struct FPropertyChangedChainEvent& PropertyChangedEvent ) override;
+#endif //WITH_EDITOR
+	
+	virtual void PostInitProperties() override;
+	//~End UObject interface
+
+	//~ Begin FTickableGameObject
+	virtual bool IsTickableInEditor() const override { return true; }
+	virtual void Tick(float DeltaTime) override;
+	virtual TStatId GetStatId() const override;
+	//~ End FTickableGameObject
+	
+	/** Returns interpolated distortion parameters based on input focus and zoom */
+	bool EvaluateDistortionParameters(float InFocus, float InZoom, FDistortionInfo& OutEvaluatedValue) const;
+
+	/** Returns interpolated intrinsic parameters based on input focus and zoom */
+	bool EvaluateIntrinsicParameters(float InFocus, float InZoom, FIntrinsicParameters& OutEvaluatedValue) const;
+
+	/** Draws the distortion map based on evaluation point*/
+	bool EvaluateDistortionData(float InFocus, float InZoom, FVector2D InFilmback, ULensDistortionModelHandlerBase* InLensHandler, FDistortionData& OutDistortionData) const;
+
+	/** Returns interpolated nodal point offset based on input focus and zoom */
+	bool EvaluateNodalPointOffset(float InFocus, float InZoom, FNodalPointOffset& OutEvaluatedValue) const;
+
+	/** Whether focus encoder mapping is configured */
+	bool HasFocusEncoderMapping() const;
+
+	/** Returns interpolated focus based on input normalized value and mapping */
+	float EvaluateNormalizedFocus(float InNormalizedValue);
+
+	/** Whether iris encoder mapping is configured */
+	bool HasIrisEncoderMapping() const;
+
+	/** Returns interpolated iris based on input normalized value and mapping */
+	float EvaluateNormalizedIris(float InNormalizedValue);
+
+	/** Whether zoom encoder mapping is configured */
+	bool HasZoomEncoderMapping() const;
+
+	/** Returns interpolated zoom based on input normalized value and mapping */
+	float EvaluateNormalizedZoom(float InNormalizedValue);
+
+	/** Callbacked when stmap derived data has completed */
+	void OnDistortionDerivedDataJobCompleted(const FDerivedDistortionDataJobOutput& JobOutput);
+
+	/** Whether the sensor dimensions in the lens file will be compatible with the sensor dimensions of the input CineCameraComponent */
+	bool IsCineCameraCompatible(const UCineCameraComponent* CineCameraComponent);
+
+protected:
+
+	/** Updates derived data entries to make sure it matches what is assigned in map points based on data mode */
+	void UpdateDerivedData();
+
+	/** Returns the overscan factor based on distorted UV and image center */
+	float ComputeOverscan(const FDistortionData& DerivedData, FVector2D PrincipalPoint) const;
+
+	/** Clears output displacement map on LensHandler to have no distortion and setup distortion data to match that */
+	void SetupNoDistortionOutput(ULensDistortionModelHandlerBase* LensHandler, FDistortionData& OutDistortionData) const;
+
+	/** Evaluates distortion based on InFocus and InZoom using parameters */
+	bool EvaluateDistortionForParameters(float InFocus, float InZoom, FVector2D InFilmback, ULensDistortionModelHandlerBase* LensHandler, FDistortionData& OutDistortionData) const;
+	
+	/** Evaluates distortion based on InFocus and InZoom using STMaps */
+	bool EvaluteDistortionForSTMaps(float InFocus, float InZoom, FVector2D InFilmback, ULensDistortionModelHandlerBase* LensHandler, FDistortionData& OutDistortionData) const;
+	
+public:
+
+	/** Lens information */
+	UPROPERTY(EditAnywhere, BlueprintReadOnly, Category = "Lens info")
+	FLensInfo LensInfo;
+
+	/** Type of data used for lens mapping */
+	UPROPERTY(EditAnywhere, Category="Lens mapping")
+	ELensDataMode DataMode = ELensDataMode::Parameters;
+	
+	/** Mapping between FIZ data and distortion parameters (k1, k2...) */
+	UPROPERTY(EditAnywhere, Category="FIZ map", meta = (EditCondition = "DataMode == ELensDataMode::Parameters"))
+	TArray<FDistortionMapPoint> DistortionMapping;
+
+	/** Mapping between FIZ data and intrinsic parameters (fx/fy and principal point) */
+	UPROPERTY(EditAnywhere, Category="FIZ map")
+	TArray<FIntrinsicMapPoint> IntrinsicMapping;
+
+	/** Precomputed data associated to a calibration point */
+	UPROPERTY(EditAnywhere, Category="FIZ map", meta = (EditCondition = "DataMode == ELensDataMode::STMap"))
+	TArray<FCalibratedMapPoint> CalibratedMapPoints;
+	
+	/** Mapping between FIZ data and nodal point */
+	UPROPERTY(EditAnywhere, Category = "FIZ map")
+	TArray<FNodalOffsetMapPoint> NodalOffsetMapping;
+
+	/** Metadata user could enter for its lens */
+	UPROPERTY(EditAnywhere, Category = "Metadata")
+	TMap<FString, FString> UserMetadata;
+
+	/** Encoder mapping from normalized value to values in physical units */
+	UPROPERTY(EditAnywhere, Category = "Encoder", AdvancedDisplay)
+	FEncoderMapping EncoderMapping;
+
+protected:
+
+	/** Derived data compute jobs we are waiting on */
+	int32 DerivedDataInFlightCount = 0;
+
+	/** Processor handling derived data out of calibrated st maps */
+	TUniquePtr<ICalibratedMapProcessor> CalibratedMapProcessor;
+
+	/** Texture used to store temporary undistortion displacement map when using map blending */
+	UPROPERTY(Transient)
+	TArray<UTextureRenderTarget2D*> UndistortionDisplacementMapHolders;
+
+	/** Texture used to store temporary distortion displacement map when using map blending */
+	UPROPERTY(Transient)
+	TArray<UTextureRenderTarget2D*> DistortionDisplacementMapHolders;
+
+	static constexpr int32 DisplacementMapHolderCount = 4;
+
+	/** UV coordinates of 8 points (4 corners + 4 mid points) */
+	static const TArray<FVector2D> UndistortedUVs;
+};
+
+
+/**
+ * Wrapper to facilitate default lensfile vs picker
+ */
+USTRUCT(BlueprintType)
+struct CAMERACALIBRATIONCORE_API FLensFilePicker
+{
+	GENERATED_BODY()
+
+public:
+
+	/** Get the proper lens whether it's the default one or the picked one */
+	ULensFile* GetLensFile() const;
+
+public:
+	/** You can override lens file to use if the default one is not desired */
+	UPROPERTY(BlueprintReadWrite, Category = "Lens File")
+	bool bOverrideDefaultLensFile = false;
+
+	UPROPERTY(EditAnywhere, BlueprintReadWrite, Category = "Lens File", Meta = (EditCondition = "bOverrideDefaultLensFile"))
+	ULensFile* LensFile = nullptr;
+};
+
+/** Structure that caches the inputs (and other useful bits) used when evaluating the Lens File */
+struct FLensFileEvalData
+{
+	FLensFileEvalData()
+	{
+		Invalidate();
+	};
+
+	ULensFile* LensFile;
+
+	/** The values that should be used as inputs to the Lut in the LensFile */
+	struct
+	{
+		/** Focus input */
+		TOptional<float> Focus;
+
+		/** Iris input */
+		TOptional<float> Iris;
+
+		/** Zoom input */
+		TOptional<float> Zoom;
+	} Input;
+
+	/** Information about the camera associated with the lens evaluation */
+	struct
+	{
+		uint32 UniqueId;
+	} Camera;
+
+	/** Information about the Distortion evaluation */
+	struct
+	{
+		/** True if distotion was applied (and the lens distortion handler updated its state) */
+		bool bWasEvaluated;
+	} Distortion;
+
+	/** Information about the nodal offset evaluation */
+	struct
+	{
+		/** True if the evaluated nodal offset was applied to the camera */
+		bool bWasApplied;
+	} NodalOffset;
+
+	/** Invalidates the data in this structure and avoid using stale or invalid values */
+	void Invalidate()
+	{
+		LensFile = nullptr;
+
+		Input.Focus.Reset();
+		Input.Iris.Reset();
+		Input.Zoom.Reset();
+
+		Camera.UniqueId = INDEX_NONE;
+		Distortion.bWasEvaluated = false;
+		NodalOffset.bWasApplied = false;
+	}
+};
+
+
diff --git a/Engine/Plugins/VirtualProduction/CameraCalibrationCore/Source/CameraCalibrationCore/Public/Models/LensModel.h b/Engine/Plugins/VirtualProduction/CameraCalibrationCore/Source/CameraCalibrationCore/Public/Models/LensModel.h
new file mode 100644
--- /dev/null
+++ b/Engine/Plugins/VirtualProduction/CameraCalibrationCore/Source/CameraCalibrationCore/Public/Models/LensModel.h
@@ -0,0 +1,76 @@
+// Copyright Epic Games, Inc. All Rights Reserved.
+
+#pragma once
+
+#include "CoreMinimal.h"
+
+#include "UObject/StructOnScope.h"
+
+#include "LensModel.generated.h"
+
+
+class ULensDistortionModelHandlerBase;
+
+
+/**
+ * Abstract base class for lens models
+ */
+UCLASS(Abstract)
+class CAMERACALIBRATIONCORE_API ULensModel : public UObject
+{
+	GENERATED_BODY()
+
+public:
+	/** Get the lens model name */
+	virtual FName GetModelName() const PURE_VIRTUAL(ULensModel::GetModelName, return FName(""););
+
+	/** Get the struct of distortion parameters supported by this model */
+	virtual UScriptStruct* GetParameterStruct() const PURE_VIRTUAL(ULensModel::GetParameterStruct, return nullptr;);
+
+#if WITH_EDITOR
+	/** Get the names of each float parameters supported by this model */
+	virtual TArray<FText> GetParameterDisplayNames() const;
+#endif //WITH_EDITOR
+
+	/** Get the number of float fields in the parameter struct supported by this model */
+	virtual uint32 GetNumParameters() const;
+
+	/** 
+	 * Fill the destination array of floats with the values of the fields in the source struct 
+	 * Note: the template type must be a UStruct
+	 */
+	template<typename StructType>
+	void ToArray(const StructType& SrcData, TArray<float>& DstArray) const
+	{
+		ToArray_Internal(StructType::StaticStruct(), &SrcData, DstArray);
+	}
+
+	/** 
+	 * ToArray specialization taking a StructOnScope containing type and data 
+	 */
+	template<>
+	void ToArray<FStructOnScope>(const FStructOnScope& SrcData, TArray<float>& DstArray) const
+	{
+		ToArray_Internal(static_cast<const UScriptStruct*>(SrcData.GetStruct()), reinterpret_cast<const void*>(SrcData.GetStructMemory()), DstArray);
+	}
+
+	/**
+	 * Populate the float fields in the destination struct with the values in the source array
+	 * Note: the template type must be a UStruct
+	 */
+	template<typename StructType>
+	void FromArray(const TArray<float>& SrcArray, StructType& DstData)
+	{
+		FromArray_Internal(StructType::StaticStruct(), SrcArray, &DstData);
+	}
+
+	/** Returns the first handler that supports the given LensModel */
+	static TSubclassOf<ULensDistortionModelHandlerBase> GetHandlerClass(TSubclassOf<ULensModel> LensModel);
+
+protected:
+	/** Internal implementation of ToArray. See declaration of public template method. */
+	virtual void ToArray_Internal(const UScriptStruct* TypeStruct, const void* SrcData, TArray<float>& DstArray) const;
+
+	/** Internal implementation of FromArray. See declaration of public template method. */
+	virtual void FromArray_Internal(UScriptStruct* TypeStruct, const TArray<float>& SrcArray, void* DstData);
+};
diff --git a/Engine/Plugins/VirtualProduction/CameraCalibrationCore/Source/CameraCalibrationCore/Public/Models/SphericalLensModel.h b/Engine/Plugins/VirtualProduction/CameraCalibrationCore/Source/CameraCalibrationCore/Public/Models/SphericalLensModel.h
new file mode 100644
--- /dev/null
+++ b/Engine/Plugins/VirtualProduction/CameraCalibrationCore/Source/CameraCalibrationCore/Public/Models/SphericalLensModel.h
@@ -0,0 +1,55 @@
+// Copyright Epic Games, Inc. All Rights Reserved.
+
+#pragma once
+
+#include "LensModel.h"
+
+#include "SphericalLensModel.generated.h"
+
+/**
+ * Spherical lens distortion parameters
+ * All parameters are unitless and represent the coefficients used to undistort a distorted image
+ * Refer to the OpenCV camera calibration documentation for the intended units/usage of these parameters:
+ * https://docs.opencv.org/3.4/d9/d0c/group__calib3d.html
+ */
+USTRUCT(BlueprintType)
+struct CAMERACALIBRATIONCORE_API FSphericalDistortionParameters
+{
+	GENERATED_BODY()
+
+public:
+	/** Radial coefficient of the r^2 term */
+	UPROPERTY(EditAnywhere, BlueprintReadWrite, Category = "Distortion")
+	float K1 = 0.0f;
+
+	/** Radial coefficient of the r^4 term */
+	UPROPERTY(EditAnywhere, BlueprintReadWrite, Category = "Distortion")
+	float K2 = 0.0f;
+
+	/** Radial coefficient of the r^6 term */
+	UPROPERTY(EditAnywhere, BlueprintReadWrite, Category = "Distortion")
+	float K3 = 0.0f;
+
+	/** First tangential coefficient */
+	UPROPERTY(EditAnywhere, BlueprintReadWrite, Category = "Distortion")
+	float P1 = 0.0f;
+
+	/** Second tangential coefficient */
+	UPROPERTY(EditAnywhere, BlueprintReadWrite, Category = "Distortion")
+	float P2 = 0.0f;
+};
+
+/**
+ * Spherical lens model, using spherical lens distortion parameters
+ */
+UCLASS(BlueprintType, meta = (DisplayName = "Spherical Lens Model"))
+class CAMERACALIBRATIONCORE_API USphericalLensModel : public ULensModel
+{
+	GENERATED_BODY()
+
+public:
+	//~ Begin ULensModel interface
+	virtual UScriptStruct* GetParameterStruct() const override;
+	virtual FName GetModelName() const override;
+	//~ End ULensModel interface
+};
diff --git a/Engine/Plugins/VirtualProduction/CameraCalibrationCore/Source/CameraCalibrationCore/Public/SphericalLensDistortionModelHandler.h b/Engine/Plugins/VirtualProduction/CameraCalibrationCore/Source/CameraCalibrationCore/Public/SphericalLensDistortionModelHandler.h
new file mode 100644
--- /dev/null
+++ b/Engine/Plugins/VirtualProduction/CameraCalibrationCore/Source/CameraCalibrationCore/Public/SphericalLensDistortionModelHandler.h
@@ -0,0 +1,29 @@
+// Copyright Epic Games, Inc. All Rights Reserved.
+
+#pragma once
+
+#include "LensDistortionModelHandlerBase.h"
+
+#include "Models/SphericalLensModel.h"
+
+#include "SphericalLensDistortionModelHandler.generated.h"
+
+/** Lens distortion handler for a spherical lens model that implements the Brown-Conrady polynomial model */
+UCLASS(BlueprintType)
+class CAMERACALIBRATIONCORE_API USphericalLensDistortionModelHandler : public ULensDistortionModelHandlerBase
+{
+	GENERATED_BODY()
+
+protected:
+	//~ Begin ULensDistortionModelHandlerBase interface
+	virtual void InitializeHandler() override;
+	virtual FVector2D ComputeDistortedUV(const FVector2D& InScreenUV) const override;
+	virtual void InitDistortionMaterials() override;
+	virtual void UpdateMaterialParameters() override;
+	virtual void InterpretDistortionParameters() override;
+	//~ End ULensDistortionModelHandlerBase interface
+
+private:
+	/** Spherical lens distortion parameters (k1, k2, k3, p1, p2) */
+	FSphericalDistortionParameters SphericalParameters;
+};
diff --git a/Engine/Plugins/VirtualProduction/CameraCalibrationCore/Source/CameraCalibrationCoreEditor/CameraCalibrationCoreEditor.Build.cs b/Engine/Plugins/VirtualProduction/CameraCalibrationCore/Source/CameraCalibrationCoreEditor/CameraCalibrationCoreEditor.Build.cs
new file mode 100644
--- /dev/null
+++ b/Engine/Plugins/VirtualProduction/CameraCalibrationCore/Source/CameraCalibrationCoreEditor/CameraCalibrationCoreEditor.Build.cs
@@ -0,0 +1,23 @@
+// Copyright Epic Games, Inc. All Rights Reserved.
+
+namespace UnrealBuildTool.Rules
+{
+	public class CameraCalibrationCoreEditor : ModuleRules
+	{
+		public CameraCalibrationCoreEditor(ReadOnlyTargetRules Target) : base(Target)
+		{
+			PrivateDependencyModuleNames.AddRange(
+				new string[]
+				{
+					"CameraCalibrationCore",
+					"Core",
+					"CoreUObject",
+					"Engine",
+					"PropertyEditor",
+					"Slate",
+					"SlateCore",
+					"UnrealEd",
+				});
+		}
+	}
+}
diff --git a/Engine/Plugins/VirtualProduction/CameraCalibrationCore/Source/CameraCalibrationCoreEditor/Private/CameraCalibrationCoreEditorLog.h b/Engine/Plugins/VirtualProduction/CameraCalibrationCore/Source/CameraCalibrationCoreEditor/Private/CameraCalibrationCoreEditorLog.h
new file mode 100644
--- /dev/null
+++ b/Engine/Plugins/VirtualProduction/CameraCalibrationCore/Source/CameraCalibrationCoreEditor/Private/CameraCalibrationCoreEditorLog.h
@@ -0,0 +1,8 @@
+// Copyright Epic Games, Inc. All Rights Reserved.
+
+#pragma once
+
+#include "CoreMinimal.h"
+#include "Logging/LogMacros.h"
+
+DECLARE_LOG_CATEGORY_EXTERN(LogCameraCalibrationCoreEditor, Log, All);
diff --git a/Engine/Plugins/VirtualProduction/CameraCalibrationCore/Source/CameraCalibrationCoreEditor/Private/CameraCalibrationCoreEditorModule.cpp b/Engine/Plugins/VirtualProduction/CameraCalibrationCore/Source/CameraCalibrationCoreEditor/Private/CameraCalibrationCoreEditorModule.cpp
new file mode 100644
--- /dev/null
+++ b/Engine/Plugins/VirtualProduction/CameraCalibrationCore/Source/CameraCalibrationCoreEditor/Private/CameraCalibrationCoreEditorModule.cpp
@@ -0,0 +1,39 @@
+// Copyright Epic Games, Inc. All Rights Reserved.
+
+
+#include "CameraCalibrationCoreEditorModule.h"
+
+#include "CoreMinimal.h"
+
+#include "DistortionHandlerPickerDetailCustomization.h"
+#include "Editor.h"
+#include "LensFile.h"
+#include "Modules/ModuleManager.h"
+#include "PropertyEditorModule.h"
+
+#define LOCTEXT_NAMESPACE "CameraCalibrationCoreEditor"
+
+DEFINE_LOG_CATEGORY(LogCameraCalibrationCoreEditor);
+
+
+void FCameraCalibrationCoreEditorModule::StartupModule()
+{
+	FPropertyEditorModule& PropertyEditorModule = FModuleManager::Get().LoadModuleChecked<FPropertyEditorModule>("PropertyEditor");
+	PropertyEditorModule.RegisterCustomPropertyTypeLayout(FDistortionHandlerPicker::StaticStruct()->GetFName(), FOnGetPropertyTypeCustomizationInstance::CreateStatic(&FDistortionHandlerPickerDetailCustomization::MakeInstance));
+}
+
+void FCameraCalibrationCoreEditorModule::ShutdownModule()
+{
+	if (!IsEngineExitRequested() && GEditor && UObjectInitialized())
+	{
+		FPropertyEditorModule& PropertyModule = FModuleManager::GetModuleChecked<FPropertyEditorModule>("PropertyEditor");
+		PropertyModule.UnregisterCustomClassLayout(ULensFile::StaticClass()->GetFName());
+		PropertyModule.UnregisterCustomPropertyTypeLayout(FDistortionHandlerPicker::StaticStruct()->GetFName());
+	}
+}
+
+
+IMPLEMENT_MODULE(FCameraCalibrationCoreEditorModule, CameraCalibrationCoreEditor);
+
+
+#undef LOCTEXT_NAMESPACE
diff --git a/Engine/Plugins/VirtualProduction/CameraCalibrationCore/Source/CameraCalibrationCoreEditor/Private/CameraCalibrationCoreEditorModule.h b/Engine/Plugins/VirtualProduction/CameraCalibrationCore/Source/CameraCalibrationCoreEditor/Private/CameraCalibrationCoreEditorModule.h
new file mode 100644
--- /dev/null
+++ b/Engine/Plugins/VirtualProduction/CameraCalibrationCore/Source/CameraCalibrationCoreEditor/Private/CameraCalibrationCoreEditorModule.h
@@ -0,0 +1,28 @@
+// Copyright Epic Games, Inc. All Rights Reserved.
+#pragma once
+
+#include "Modules/ModuleInterface.h"
+
+#include "CoreMinimal.h"
+
+class IAssetTypeActions;
+
+DECLARE_LOG_CATEGORY_EXTERN(LogCameraCalibrationCoreEditor, Log, All);
+
+
+/**
+ * Implements the CameraCalibrationCoreEditor module.
+ */
+class FCameraCalibrationCoreEditorModule : public IModuleInterface
+{
+public:
+
+	//~ Begin IModuleInterface interface
+	virtual void StartupModule() override;
+	virtual void ShutdownModule() override;
+	//~ End IModuleInterface interface
+
+private:
+
+	TArray<TSharedRef<IAssetTypeActions>> RegisteredAssetTypeActions;
+};
diff --git a/Engine/Plugins/VirtualProduction/CameraCalibrationCore/Source/CameraCalibrationCoreEditor/Private/DistortionHandlerPickerDetailCustomization.cpp b/Engine/Plugins/VirtualProduction/CameraCalibrationCore/Source/CameraCalibrationCoreEditor/Private/DistortionHandlerPickerDetailCustomization.cpp
new file mode 100644
--- /dev/null
+++ b/Engine/Plugins/VirtualProduction/CameraCalibrationCore/Source/CameraCalibrationCoreEditor/Private/DistortionHandlerPickerDetailCustomization.cpp
@@ -0,0 +1,164 @@
+// Copyright Epic Games, Inc. All Rights Reserved.
+
+#include "DistortionHandlerPickerDetailCustomization.h"
+
+#include "CameraCalibrationSubsystem.h"
+#include "CameraCalibrationTypes.h"
+#include "DetailWidgetRow.h"
+#include "Framework/MultiBox/MultiBoxBuilder.h"
+#include "IPropertyUtilities.h"
+#include "Misc/Guid.h"
+
+
+#define LOCTEXT_NAMESPACE "DistortionHandlerPickerDetailCustomization"
+
+
+TSharedRef<IPropertyTypeCustomization> FDistortionHandlerPickerDetailCustomization::MakeInstance()
+{
+	return MakeShared<FDistortionHandlerPickerDetailCustomization>();
+}
+
+void FDistortionHandlerPickerDetailCustomization::CustomizeHeader(TSharedRef<IPropertyHandle> InPropertyHandle, FDetailWidgetRow& HeaderRow, IPropertyTypeCustomizationUtils& CustomizationUtils)
+{
+	StructPropertyHandle = InPropertyHandle;
+	TSharedPtr<IPropertyUtilities> PropertyUtils = CustomizationUtils.GetPropertyUtilities();
+	
+	StructPropertyHandle->MarkResetToDefaultCustomized();
+
+	TArray<void*> RawData;
+	StructPropertyHandle->AccessRawData(RawData);
+
+	// Multi-select is not supported for this property
+	if (RawData.Num() == 1)
+	{
+		FDistortionHandlerPicker* DistortionHandlerPicker = reinterpret_cast<FDistortionHandlerPicker*>(RawData[0]);
+
+		HeaderRow.NameContent()
+			[
+				StructPropertyHandle->CreatePropertyNameWidget()
+			]
+			.ValueContent()
+			[
+				BuildDistortionHandlerPickerWidget(DistortionHandlerPicker)
+			]
+			.IsEnabled(MakeAttributeLambda([=] { return !InPropertyHandle->IsEditConst() && PropertyUtils->IsPropertyEditingEnabled(); }));
+	}
+}
+
+TSharedRef<SWidget> FDistortionHandlerPickerDetailCustomization::BuildDistortionHandlerPickerWidget(FDistortionHandlerPicker* InDistortionHandlerPicker)
+{
+	TSharedRef<SWidget> ValueWidget =
+		SNew(SHorizontalBox)
+	 	+ SHorizontalBox::Slot()
+ 		.AutoWidth()
+ 		.Padding(FMargin(0.0f, 0.0f, 0.0f, 0.0f))
+ 		[
+ 			SNew(SComboButton)
+ 			.OnGetMenuContent(this, &FDistortionHandlerPickerDetailCustomization::PopulateDistortionHandlerComboButton, InDistortionHandlerPicker)
+	 		.ContentPadding(FMargin(4.0, 2.0))
+			.ButtonContent()
+			[
+				SNew(STextBlock)
+				.Text(this, &FDistortionHandlerPickerDetailCustomization::OnGetButtonText, InDistortionHandlerPicker)
+			]
+		]
+	;
+
+	return ValueWidget;
+}
+
+FText FDistortionHandlerPickerDetailCustomization::OnGetButtonText(FDistortionHandlerPicker* InDistortionHandlerPicker) const
+{
+	if (InDistortionHandlerPicker->HandlerDisplayName.IsEmpty())
+	{
+		return FText::FromName(NAME_None);
+	}
+	return FText::FromString(InDistortionHandlerPicker->HandlerDisplayName);
+}
+
+TSharedRef<SWidget> FDistortionHandlerPickerDetailCustomization::PopulateDistortionHandlerComboButton(FDistortionHandlerPicker* InDistortionHandlerPicker) const
+{
+	// Generate menu
+	FMenuBuilder MenuBuilder(true, nullptr);
+	MenuBuilder.BeginSection("DistortionHandlers", LOCTEXT("DistortionHandlers", "Distortion Sources For Target Camera"));
+	{
+		TArray<ULensDistortionModelHandlerBase*> Handlers;
+		UCameraCalibrationSubsystem* SubSystem = GEngine->GetEngineSubsystem<UCameraCalibrationSubsystem>();
+		if (SubSystem)
+		{
+			Handlers = SubSystem->GetDistortionModelHandlers(InDistortionHandlerPicker->TargetCameraComponent);
+		}
+
+		if (Handlers.Num() == 0)
+		{
+			ULensDistortionModelHandlerBase* NullHandler = nullptr;
+			FString EmptyString;
+
+			//Always add a None entry
+			MenuBuilder.AddMenuEntry(
+				FText::FromName(NAME_None),
+				FText::FromName(NAME_None),
+				FSlateIcon(),
+				FUIAction(
+					FExecuteAction::CreateSP(this, &FDistortionHandlerPickerDetailCustomization::OnDistortionHandlerSelected, InDistortionHandlerPicker, NullHandler),
+					FCanExecuteAction(),
+					FIsActionChecked::CreateSP(this, &FDistortionHandlerPickerDetailCustomization::IsDistortionHandlerSelected, InDistortionHandlerPicker, EmptyString)
+				),
+				NAME_None,
+				EUserInterfaceActionType::RadioButton
+			);
+		}
+
+		for (ULensDistortionModelHandlerBase* Handler : Handlers)
+		{
+			MenuBuilder.AddMenuEntry(
+				FText::Format(LOCTEXT("DistortionHandlerLabel", "{0}"), FText::FromString(Handler->GetDisplayName())),
+				FText::Format(LOCTEXT("DistortionHandlerTooltip", "{0}"), FText::FromString(Handler->GetDisplayName())),
+				FSlateIcon(),
+				FUIAction(
+					FExecuteAction::CreateSP(this, &FDistortionHandlerPickerDetailCustomization::OnDistortionHandlerSelected, InDistortionHandlerPicker, Handler),
+					FCanExecuteAction(),
+					FIsActionChecked::CreateSP(this, &FDistortionHandlerPickerDetailCustomization::IsDistortionHandlerSelected, InDistortionHandlerPicker, Handler->GetDisplayName())
+				),
+				NAME_None,
+				EUserInterfaceActionType::RadioButton
+			);
+		}
+	}
+	MenuBuilder.EndSection();
+
+	return MenuBuilder.MakeWidget();
+}
+
+void FDistortionHandlerPickerDetailCustomization::OnDistortionHandlerSelected(FDistortionHandlerPicker* InDistortionHandlerPicker, ULensDistortionModelHandlerBase* InHandler) const
+{
+	const FScopedTransaction Transaction(LOCTEXT("SetDistortionHandler", "Set Distortion Source"));
+
+	TArray<UObject*> OuterObjects;
+	StructPropertyHandle->GetOuterObjects(OuterObjects);
+	for (UObject* Outer : OuterObjects)
+	{
+		Outer->Modify();
+	}
+
+	FGuid DistortionProducerID;
+	FString HandlerDisplayName;
+
+	if (InHandler)
+	{
+		DistortionProducerID = InHandler->GetDistortionProducerID();
+		HandlerDisplayName = InHandler->GetDisplayName();
+	}
+
+	InDistortionHandlerPicker->DistortionProducerID = DistortionProducerID;
+	InDistortionHandlerPicker->HandlerDisplayName = HandlerDisplayName;
+
+	StructPropertyHandle->NotifyPostChange(EPropertyChangeType::ValueSet);
+}
+
+bool FDistortionHandlerPickerDetailCustomization::IsDistortionHandlerSelected(FDistortionHandlerPicker* InDistortionHandlerPicker, FString InDisplayName) const
+{
+	return (InDistortionHandlerPicker->HandlerDisplayName.Equals(InDisplayName));
+}
+
+#undef LOCTEXT_NAMESPACE
\ No newline at end of file
diff --git a/Engine/Plugins/VirtualProduction/CameraCalibrationCore/Source/CameraCalibrationCoreEditor/Private/DistortionHandlerPickerDetailCustomization.h b/Engine/Plugins/VirtualProduction/CameraCalibrationCore/Source/CameraCalibrationCoreEditor/Private/DistortionHandlerPickerDetailCustomization.h
new file mode 100644
--- /dev/null
+++ b/Engine/Plugins/VirtualProduction/CameraCalibrationCore/Source/CameraCalibrationCoreEditor/Private/DistortionHandlerPickerDetailCustomization.h
@@ -0,0 +1,46 @@
+// Copyright Epic Games, Inc. All Rights Reserved.
+
+#pragma once
+
+#include "IPropertyTypeCustomization.h"
+
+#include "CameraCalibrationTypes.h"
+#include "CoreMinimal.h"
+#include "ScopedTransaction.h"
+#include "Widgets/Text/STextBlock.h"
+
+class UCineCameraComponent;
+class ULensDistortionModelHandlerBase;
+
+class FDistortionHandlerPickerDetailCustomization : public IPropertyTypeCustomization
+{
+public:
+	/** Makes a new instance of this detail layout class for a specific detail view requesting it. */
+	static TSharedRef<IPropertyTypeCustomization> MakeInstance();
+
+public:
+	//~ Begin IDetailCustomization interface
+	virtual void CustomizeHeader(TSharedRef<IPropertyHandle> PropertyHandle, FDetailWidgetRow& HeaderRow, IPropertyTypeCustomizationUtils& CustomizationUtils) override;
+	virtual void CustomizeChildren(TSharedRef<IPropertyHandle> PropertyHandle, IDetailChildrenBuilder& ChildBuilder, IPropertyTypeCustomizationUtils& CustomizationUtils) override {}
+	//~ End IDetailCustomization interface
+
+protected:
+	/** Construct the combo button widget to display a dropdown of the distortion sources that exist for the input picker */
+	TSharedRef<SWidget> BuildDistortionHandlerPickerWidget(FDistortionHandlerPicker* InDistortionHandlerPicker);
+
+	/** Construct a menu widget of distortion sources that exist for the input picker */
+	TSharedRef<SWidget> PopulateDistortionHandlerComboButton(FDistortionHandlerPicker* InDistortionHandlerPicker) const;
+
+	/** Delegate to run when a user makes a selection in the distortion source dropdown */
+	void OnDistortionHandlerSelected(FDistortionHandlerPicker* InDistortionHandlerPicker, ULensDistortionModelHandlerBase* InHandler) const;
+
+	/** Delegate to run to determine to determine whether the radio button should be filled in next to each dropdown option */
+	bool IsDistortionHandlerSelected(FDistortionHandlerPicker* InDistortionHandlerPicker, FString InDisplayName) const;
+
+	/** Delegate to run to get the text to display on the combo button */
+	FText OnGetButtonText(FDistortionHandlerPicker* InDistortionHandlerPicker) const;
+
+protected:
+	/** Cached property handle to the struct being customized */
+	TSharedPtr<IPropertyHandle> StructPropertyHandle;
+};
diff --git a/Engine/Plugins/VirtualProduction/LiveLinkCamera/LiveLinkCamera.uplugin b/Engine/Plugins/VirtualProduction/LiveLinkCamera/LiveLinkCamera.uplugin
--- a/Engine/Plugins/VirtualProduction/LiveLinkCamera/LiveLinkCamera.uplugin
+++ b/Engine/Plugins/VirtualProduction/LiveLinkCamera/LiveLinkCamera.uplugin
@@ -20,7 +20,7 @@
 			"Enabled": true
 		},
 		{
-			"Name": "CameraCalibration",
+			"Name": "CameraCalibrationCore",
 			"Enabled": true
 		}
 	],
diff --git a/Engine/Plugins/VirtualProduction/LiveLinkCamera/Source/LiveLinkCamera/LiveLinkCamera.build.cs b/Engine/Plugins/VirtualProduction/LiveLinkCamera/Source/LiveLinkCamera/LiveLinkCamera.build.cs
--- a/Engine/Plugins/VirtualProduction/LiveLinkCamera/Source/LiveLinkCamera/LiveLinkCamera.build.cs
+++ b/Engine/Plugins/VirtualProduction/LiveLinkCamera/Source/LiveLinkCamera/LiveLinkCamera.build.cs
@@ -24,7 +24,7 @@ public class LiveLinkCamera : ModuleRules
 		PublicDependencyModuleNames.AddRange(
 			new string[]
 			{
-				"CameraCalibration",
+				"CameraCalibrationCore",
 				"CinematicCamera",
 				"Core",
 				"CoreUObject",
diff --git a/Engine/Plugins/VirtualProduction/LiveLinkCamera/Source/LiveLinkCameraEditor/LiveLinkCameraEditor.build.cs b/Engine/Plugins/VirtualProduction/LiveLinkCamera/Source/LiveLinkCameraEditor/LiveLinkCameraEditor.build.cs
--- a/Engine/Plugins/VirtualProduction/LiveLinkCamera/Source/LiveLinkCameraEditor/LiveLinkCameraEditor.build.cs
+++ b/Engine/Plugins/VirtualProduction/LiveLinkCamera/Source/LiveLinkCameraEditor/LiveLinkCameraEditor.build.cs
@@ -30,7 +30,7 @@ public class LiveLinkCameraEditor : ModuleRules
 		PrivateDependencyModuleNames.AddRange(
 			new string[]
 			{
-				"CameraCalibration",
+				"CameraCalibrationCore",
 				"Core",
 				"CoreUObject",
 				"DetailCustomizations",
diff --git a/Engine/Plugins/VirtualProduction/LiveLinkLens/LiveLinkLens.uplugin b/Engine/Plugins/VirtualProduction/LiveLinkLens/LiveLinkLens.uplugin
--- a/Engine/Plugins/VirtualProduction/LiveLinkLens/LiveLinkLens.uplugin
+++ b/Engine/Plugins/VirtualProduction/LiveLinkLens/LiveLinkLens.uplugin
@@ -26,7 +26,7 @@
 			"Enabled": true
 		},
 		{
-			"Name": "CameraCalibration",
+			"Name": "CameraCalibrationCore",
 			"Enabled": true
 		}
 	]
diff --git a/Engine/Plugins/VirtualProduction/LiveLinkLens/Source/LiveLinkLens/LiveLinkLens.Build.cs b/Engine/Plugins/VirtualProduction/LiveLinkLens/Source/LiveLinkLens/LiveLinkLens.Build.cs
--- a/Engine/Plugins/VirtualProduction/LiveLinkLens/Source/LiveLinkLens/LiveLinkLens.Build.cs
+++ b/Engine/Plugins/VirtualProduction/LiveLinkLens/Source/LiveLinkLens/LiveLinkLens.Build.cs
@@ -11,7 +11,7 @@ public class LiveLinkLens : ModuleRules
 		PublicDependencyModuleNames.AddRange(
 			new string[]
 			{
-				"CameraCalibration",
+				"CameraCalibrationCore",
 				"CinematicCamera",
 				"LiveLinkComponents",
 				"LiveLinkInterface"
