Lens Calibration. Unit Test for Remove Points

diff --git a/Engine/Plugins/VirtualProduction/CameraCalibrationCore/Source/CameraCalibrationCore/Private/LensFile.cpp b/Engine/Plugins/VirtualProduction/CameraCalibrationCore/Source/CameraCalibrationCore/Private/LensFile.cpp
--- a/Engine/Plugins/VirtualProduction/CameraCalibrationCore/Source/CameraCalibrationCore/Private/LensFile.cpp
+++ b/Engine/Plugins/VirtualProduction/CameraCalibrationCore/Source/CameraCalibrationCore/Private/LensFile.cpp
@@ -1224,40 +1224,45 @@ void ULensFile::ClearData(ELensDataCategory InDataCategory)
 }
 
 bool ULensFile::HasSamples(ELensDataCategory InDataCategory) const
+{
+	return GetTotalPointNum(InDataCategory) > 0 ? true : false;
+}
+
+int32 ULensFile::GetTotalPointNum(ELensDataCategory InDataCategory) const
 {
 	switch(InDataCategory)
 	{
 		case ELensDataCategory::Distortion:
 		{
-			return DistortionTable.GetFocusPoints().Num() > 0;
+			return DistortionTable.GetTotalPointNum();
 		}
 		case ELensDataCategory::ImageCenter:
 		{
-			return ImageCenterTable.GetFocusPoints().Num() > 0;
+			return ImageCenterTable.GetTotalPointNum();
 		}
 		case ELensDataCategory::Zoom:
 		{
-			return FocalLengthTable.GetFocusPoints().Num() > 0;
+			return FocalLengthTable.GetTotalPointNum();
 		}
 		case ELensDataCategory::STMap:
 		{
-			return STMapTable.GetFocusPoints().Num() > 0;
+			return STMapTable.GetTotalPointNum();
 		}
 		case ELensDataCategory::NodalOffset:
 		{
-			return NodalOffsetTable.GetFocusPoints().Num() > 0;
+			return NodalOffsetTable.GetTotalPointNum();
 		}
 		case ELensDataCategory::Focus:
 		{
-			return EncodersTable.GetNumFocusPoints() > 0;
+			return EncodersTable.GetNumFocusPoints();
 		}
 		case ELensDataCategory::Iris:
 		{
-			return EncodersTable.GetNumIrisPoints() > 0;
+			return EncodersTable.GetNumIrisPoints();
 		}
 		default:
 		{
-			return false;
+			return -1;
 		}
 	}
 }
diff --git a/Engine/Plugins/VirtualProduction/CameraCalibrationCore/Source/CameraCalibrationCore/Private/Tables/DistortionParametersTable.cpp b/Engine/Plugins/VirtualProduction/CameraCalibrationCore/Source/CameraCalibrationCore/Private/Tables/DistortionParametersTable.cpp
--- a/Engine/Plugins/VirtualProduction/CameraCalibrationCore/Source/CameraCalibrationCore/Private/Tables/DistortionParametersTable.cpp
+++ b/Engine/Plugins/VirtualProduction/CameraCalibrationCore/Source/CameraCalibrationCore/Private/Tables/DistortionParametersTable.cpp
@@ -103,6 +103,11 @@ void FDistortionFocusPoint::SetParameterValue(int32 InZoomIndex, float InZoomVal
 	}
 }
 
+int32 FDistortionTable::GetTotalPointNum() const
+{
+	return LensDataTableUtils::GetTotalPointNum(FocusPoints);
+}
+
 UScriptStruct* FDistortionTable::GetScriptStruct() const
 {
 	return StaticStruct();
diff --git a/Engine/Plugins/VirtualProduction/CameraCalibrationCore/Source/CameraCalibrationCore/Private/Tables/EncodersTable.cpp b/Engine/Plugins/VirtualProduction/CameraCalibrationCore/Source/CameraCalibrationCore/Private/Tables/EncodersTable.cpp
--- a/Engine/Plugins/VirtualProduction/CameraCalibrationCore/Source/CameraCalibrationCore/Private/Tables/EncodersTable.cpp
+++ b/Engine/Plugins/VirtualProduction/CameraCalibrationCore/Source/CameraCalibrationCore/Private/Tables/EncodersTable.cpp
@@ -2,7 +2,7 @@
 
 
 #include "Tables/EncodersTable.h"
-
+#include "CameraCalibrationCoreLog.h"
 
 int32 FEncodersTable::GetNumFocusPoints() const
 {
@@ -11,11 +11,23 @@ int32 FEncodersTable::GetNumFocusPoints() const
 
 float FEncodersTable::GetFocusInput(int32 Index) const
 {
+	if (!ensure(Focus.Keys.IsValidIndex(Index)))
+	{
+		UE_LOG(LogCameraCalibrationCore, Error, TEXT("Invalid focus keys index: %d"), Index);
+		return 0.f;
+	}
+	
 	return Focus.Keys[Index].Time;
 }
 
 float FEncodersTable::GetFocusValue(int32 Index) const
 {
+	if (!ensure(Focus.Keys.IsValidIndex(Index)))
+	{
+		UE_LOG(LogCameraCalibrationCore, Error, TEXT("Invalid focus keys index: %d"), Index);
+		return 0.f;
+	}
+	
 	return Focus.Keys[Index].Value;
 }
 
@@ -38,11 +50,23 @@ int32 FEncodersTable::GetNumIrisPoints() const
 
 float FEncodersTable::GetIrisInput(int32 Index) const
 {
+	if (!ensure(Iris.Keys.IsValidIndex(Index)))
+	{
+		UE_LOG(LogCameraCalibrationCore, Error, TEXT("Invalid iris keys index: %d"), Index);
+		return 0.f;
+	}
+
 	return Iris.Keys[Index].Time;
 }
 
 float FEncodersTable::GetIrisValue(int32 Index) const
 {
+	if (!ensure(Iris.Keys.IsValidIndex(Index)))
+	{
+		UE_LOG(LogCameraCalibrationCore, Error, TEXT("Invalid iris keys index: %d"), Index);
+		return 0.f;
+	}
+
 	return Iris.Keys[Index].Value;
 }
 
diff --git a/Engine/Plugins/VirtualProduction/CameraCalibrationCore/Source/CameraCalibrationCore/Private/Tables/FocalLengthTable.cpp b/Engine/Plugins/VirtualProduction/CameraCalibrationCore/Source/CameraCalibrationCore/Private/Tables/FocalLengthTable.cpp
--- a/Engine/Plugins/VirtualProduction/CameraCalibrationCore/Source/CameraCalibrationCore/Private/Tables/FocalLengthTable.cpp
+++ b/Engine/Plugins/VirtualProduction/CameraCalibrationCore/Source/CameraCalibrationCore/Private/Tables/FocalLengthTable.cpp
@@ -152,6 +152,11 @@ TMap<ELensDataCategory, FLinkPointMetadata> FFocalLengthTable::GetLinkedCategori
 	return LinkedToCategories;
 }
 
+int32 FFocalLengthTable::GetTotalPointNum() const
+{
+	return LensDataTableUtils::GetTotalPointNum(FocusPoints);
+}
+
 UScriptStruct* FFocalLengthTable::GetScriptStruct() const
 {
 	return StaticStruct();
diff --git a/Engine/Plugins/VirtualProduction/CameraCalibrationCore/Source/CameraCalibrationCore/Private/Tables/ImageCenterTable.cpp b/Engine/Plugins/VirtualProduction/CameraCalibrationCore/Source/CameraCalibrationCore/Private/Tables/ImageCenterTable.cpp
--- a/Engine/Plugins/VirtualProduction/CameraCalibrationCore/Source/CameraCalibrationCore/Private/Tables/ImageCenterTable.cpp
+++ b/Engine/Plugins/VirtualProduction/CameraCalibrationCore/Source/CameraCalibrationCore/Private/Tables/ImageCenterTable.cpp
@@ -108,6 +108,11 @@ bool FImageCenterTable::DoesFocusPointExists(float InFocus) const
 	return false;
 }
 
+int32 FImageCenterTable::GetTotalPointNum() const
+{
+	return LensDataTableUtils::GetTotalPointNum(FocusPoints);
+}
+
 UScriptStruct* FImageCenterTable::GetScriptStruct() const
 {
 	return StaticStruct();
diff --git a/Engine/Plugins/VirtualProduction/CameraCalibrationCore/Source/CameraCalibrationCore/Private/Tables/LensTableUtils.h b/Engine/Plugins/VirtualProduction/CameraCalibrationCore/Source/CameraCalibrationCore/Private/Tables/LensTableUtils.h
--- a/Engine/Plugins/VirtualProduction/CameraCalibrationCore/Source/CameraCalibrationCore/Private/Tables/LensTableUtils.h
+++ b/Engine/Plugins/VirtualProduction/CameraCalibrationCore/Source/CameraCalibrationCore/Private/Tables/LensTableUtils.h
@@ -234,4 +234,18 @@ namespace LensDataTableUtils
 
 		return Container[FocusNeighbors.PreviousIndex].GetValue(ZoomNeighbors.PreviousIndex, OutData);
 	}
+
+	/** Get total number of Zoom points for all Focus points of this data table */
+	template<typename FocusPointType>
+	int32 GetTotalPointNum(const TArray<FocusPointType>& Container)
+	{
+		int32 PointNum = 0;
+	
+		for (const FocusPointType& Point : Container)
+		{
+			PointNum += Point.GetNumPoints();
+		}
+
+		return PointNum;
+	}
 }
diff --git a/Engine/Plugins/VirtualProduction/CameraCalibrationCore/Source/CameraCalibrationCore/Private/Tables/NodalOffsetTable.cpp b/Engine/Plugins/VirtualProduction/CameraCalibrationCore/Source/CameraCalibrationCore/Private/Tables/NodalOffsetTable.cpp
--- a/Engine/Plugins/VirtualProduction/CameraCalibrationCore/Source/CameraCalibrationCore/Private/Tables/NodalOffsetTable.cpp
+++ b/Engine/Plugins/VirtualProduction/CameraCalibrationCore/Source/CameraCalibrationCore/Private/Tables/NodalOffsetTable.cpp
@@ -175,6 +175,11 @@ TMap<ELensDataCategory, FLinkPointMetadata> FNodalOffsetTable::GetLinkedCategori
 	return LinkedToCategories;
 }
 
+int32 FNodalOffsetTable::GetTotalPointNum() const
+{
+	return LensDataTableUtils::GetTotalPointNum(FocusPoints);
+}
+
 UScriptStruct* FNodalOffsetTable::GetScriptStruct() const
 {
 	return StaticStruct();
diff --git a/Engine/Plugins/VirtualProduction/CameraCalibrationCore/Source/CameraCalibrationCore/Private/Tables/STMapTable.cpp b/Engine/Plugins/VirtualProduction/CameraCalibrationCore/Source/CameraCalibrationCore/Private/Tables/STMapTable.cpp
--- a/Engine/Plugins/VirtualProduction/CameraCalibrationCore/Source/CameraCalibrationCore/Private/Tables/STMapTable.cpp
+++ b/Engine/Plugins/VirtualProduction/CameraCalibrationCore/Source/CameraCalibrationCore/Private/Tables/STMapTable.cpp
@@ -149,6 +149,11 @@ const FBaseFocusPoint* FSTMapTable::GetBaseFocusPoint(int32 InIndex) const
 	return nullptr;
 }
 
+int32 FSTMapTable::GetTotalPointNum() const
+{
+	return LensDataTableUtils::GetTotalPointNum(FocusPoints);
+}
+
 UScriptStruct* FSTMapTable::GetScriptStruct() const
 {
 	return StaticStruct();
diff --git a/Engine/Plugins/VirtualProduction/CameraCalibrationCore/Source/CameraCalibrationCore/Private/Tests/TestDistortionMapping.cpp b/Engine/Plugins/VirtualProduction/CameraCalibrationCore/Source/CameraCalibrationCore/Private/Tests/TestDistortionMapping.cpp
--- a/Engine/Plugins/VirtualProduction/CameraCalibrationCore/Source/CameraCalibrationCore/Private/Tests/TestDistortionMapping.cpp
+++ b/Engine/Plugins/VirtualProduction/CameraCalibrationCore/Source/CameraCalibrationCore/Private/Tests/TestDistortionMapping.cpp
@@ -305,6 +305,120 @@ namespace CameraCalibrationTestUtil
 		Test.TestEqual(FString::Printf(TEXT("FocusPoints[0], ZoomPoints[0]")), FocusPoints[2].ZoomPoints[0].Zoom, 0.0f);
 		Test.TestEqual(FString::Printf(TEXT("FocusPoints[0], ZoomPoints[1]")), FocusPoints[2].ZoomPoints[1].Zoom, 0.5f);
 	}
+
+	template <bool Const, typename Class, typename FuncType>
+    struct TClassFunPtrType;
+    
+    template <typename Class, typename RetType, typename... ArgTypes>
+    struct TClassFunPtrType<false, Class, RetType(ArgTypes...)>
+    {
+    	typedef RetType (Class::* Type)(ArgTypes...);
+    };
+    
+    template <typename Class, typename RetType, typename... ArgTypes>
+    struct TClassFunPtrType<true, Class, RetType(ArgTypes...)>
+    {
+    	typedef RetType (Class::* Type)(ArgTypes...) const;
+    };
+
+	template <typename FPointType>
+	void TestRemoveSinglePoint(FAutomationTestBase& Test, ULensFile* InLensFile, ELensDataCategory InDataCategory,
+	                       typename TClassFunPtrType<false, ULensFile, void(float, float, const FPointType&)>::Type
+	                       InFunc)
+	{
+		UEnum* Enum = StaticEnum<ELensDataCategory>();
+		const FString EnumString = Enum ? Enum->GetNameByIndex(static_cast<uint8>(InDataCategory)).ToString() : "";
+		
+		Test.AddInfo(FString::Printf(TEXT("Remove %s"), *EnumString));
+		const FPointType PointInfo;
+		const float FocusValue = 0.12f;
+		const float ZoomValue = 0.91f;
+		(InLensFile->*InFunc)(FocusValue, ZoomValue, PointInfo);
+		InLensFile->RemoveZoomPoint(InDataCategory, FocusValue, ZoomValue);
+		const int32 PointsNum = InLensFile->GetTotalPointNum(InDataCategory);
+		Test.TestEqual(TEXT("Num Zoom Points"), PointsNum, 0);
+	}
+
+	void TestLensFileRemovePoints(FAutomationTestBase& Test)
+	{
+		constexpr auto LensFileName = TEXT("RemovePointsTestLensFile");
+		ULensFile* LensFile = NewObject<ULensFile>(GetTransientPackage(), LensFileName);
+
+		Test.AddInfo(TEXT("Remove ELensDataCategory::Focus"));
+		{
+			constexpr float FocusKeyTimeValue = 0.12f;
+			constexpr float FocusKeyMappingValue = 0.91f;
+			LensFile->EncodersTable.Focus.AddKey(FocusKeyTimeValue, FocusKeyMappingValue);
+			Test.TestEqual(TEXT("Time Value Should be the same"), LensFile->EncodersTable.GetFocusInput(0), FocusKeyTimeValue);
+			LensFile->EncodersTable.RemoveFocusPoint(FocusKeyTimeValue);
+			Test.TestEqual(TEXT("Num Focus Points"), LensFile->EncodersTable.GetNumFocusPoints(), 0);
+		}
+
+		Test.AddInfo(TEXT("Remove ELensDataCategory::Iris"));
+		{
+			constexpr float IrisKeyTimeValue = 0.12f;
+			constexpr float IrisKeyMappingValue = 0.91f;
+			LensFile->EncodersTable.Iris.AddKey(IrisKeyTimeValue, IrisKeyMappingValue);
+			Test.TestEqual(TEXT("Time Value should be the same"), LensFile->EncodersTable.GetIrisInput(0), IrisKeyTimeValue);
+			LensFile->EncodersTable.RemoveIrisPoint(IrisKeyTimeValue);
+			Test.TestEqual(TEXT("Num Iris Points"), LensFile->EncodersTable.GetNumIrisPoints(), 0);
+		}
+		
+		Test.AddInfo(TEXT("Remove all point ELensDataCategory::Zoom"));
+		{
+			const FFocalLengthInfo FocalLengthInfo;
+			const float FocusValue = 0.12f;
+			const float ZoomValue1 = 0.191f;
+			const float ZoomValue2 = 0.91f;
+			LensFile->AddFocalLengthPoint(FocusValue, ZoomValue1, FocalLengthInfo);
+			LensFile->AddFocalLengthPoint(FocusValue, ZoomValue2, FocalLengthInfo);
+			LensFile->RemoveFocusPoint(ELensDataCategory::Zoom, FocusValue);
+			const int32 PointsNum = LensFile->GetTotalPointNum(ELensDataCategory::Zoom);
+			Test.TestEqual(TEXT("Num Zoom Points"), PointsNum, 0);
+		}
+
+		Test.AddInfo(TEXT("Remove ELensDataCategory::Distortion"));
+		{
+			const FFocalLengthInfo FocalLengthInfo;
+			const FDistortionInfo DistortionPoint;
+			const float FocusValue = 0.12f;
+			const float ZoomValue = 0.191f;
+			LensFile->AddDistortionPoint(FocusValue, ZoomValue, DistortionPoint, FocalLengthInfo);
+			LensFile->RemoveZoomPoint(ELensDataCategory::Distortion, FocusValue, ZoomValue);
+			const int32 DistortionPointsNum = LensFile->GetTotalPointNum(ELensDataCategory::Distortion);
+			int32 FocalLengthPointsNum = LensFile->GetTotalPointNum(ELensDataCategory::Zoom);
+			Test.TestEqual(TEXT("Num Distortion Points"), DistortionPointsNum, 0);
+			Test.TestEqual(TEXT("Num FocalLength Points"), FocalLengthPointsNum, 1);
+			LensFile->RemoveZoomPoint(ELensDataCategory::Zoom, FocusValue, ZoomValue);
+			FocalLengthPointsNum = LensFile->GetTotalPointNum(ELensDataCategory::Zoom);
+			Test.TestEqual(TEXT("Num FocalLength Points"), FocalLengthPointsNum, 0);
+		}
+
+		Test.AddInfo(TEXT("Remove ELensDataCategory::Distortion, add focal length before add distortion point"));
+		{
+			const FFocalLengthInfo FocalLengthInfo;
+			const FDistortionInfo DistortionPoint;
+			const float FocusValue = 0.12f;
+			const float ZoomValue = 0.191f;
+			LensFile->AddFocalLengthPoint(FocusValue, ZoomValue, FocalLengthInfo);
+			LensFile->AddDistortionPoint(FocusValue, ZoomValue, DistortionPoint, FocalLengthInfo);
+			int32 FocalLengthPointsNum = LensFile->GetTotalPointNum(ELensDataCategory::Zoom);
+			Test.TestEqual(TEXT("Num FocalLength Points"), FocalLengthPointsNum, 1);
+			LensFile->RemoveZoomPoint(ELensDataCategory::Distortion, FocusValue, ZoomValue);
+			const int32 DistortionPointsNum = LensFile->GetTotalPointNum(ELensDataCategory::Distortion);
+			FocalLengthPointsNum = LensFile->GetTotalPointNum(ELensDataCategory::Zoom);
+			Test.TestEqual(TEXT("Num Distortion Points"), DistortionPointsNum, 0);
+			Test.TestEqual(TEXT("Num FocalLength Points"), FocalLengthPointsNum, 1);
+			LensFile->RemoveZoomPoint(ELensDataCategory::Zoom, FocusValue, ZoomValue);
+			FocalLengthPointsNum = LensFile->GetTotalPointNum(ELensDataCategory::Zoom);
+			Test.TestEqual(TEXT("Num FocalLength Points"), FocalLengthPointsNum, 0);
+		}
+
+		TestRemoveSinglePoint<FFocalLengthInfo>(Test, LensFile, ELensDataCategory::Zoom, &ULensFile::AddFocalLengthPoint);
+		TestRemoveSinglePoint<FImageCenterInfo>(Test, LensFile, ELensDataCategory::ImageCenter, &ULensFile::AddImageCenterPoint);
+		TestRemoveSinglePoint<FNodalPointOffset>(Test, LensFile, ELensDataCategory::NodalOffset, &ULensFile::AddNodalOffsetPoint);
+		TestRemoveSinglePoint<FSTMapInfo>(Test, LensFile, ELensDataCategory::STMap, &ULensFile::AddSTMapPoint);
+	}
 }
 
 
@@ -312,8 +426,10 @@ bool FTestCameraCalibrationCore::RunTest(const FString& Parameters)
 {
 	CameraCalibrationTestUtil::TestDistortionParameterCurveBlending(*this);
 	CameraCalibrationTestUtil::TestLensFileAddPoints(*this);
+	CameraCalibrationTestUtil::TestLensFileRemovePoints(*this);
 	return true;
 }
+
 #endif // WITH_DEV_AUTOMATION_TESTS
 
 
diff --git a/Engine/Plugins/VirtualProduction/CameraCalibrationCore/Source/CameraCalibrationCore/Public/LensFile.h b/Engine/Plugins/VirtualProduction/CameraCalibrationCore/Source/CameraCalibrationCore/Public/LensFile.h
--- a/Engine/Plugins/VirtualProduction/CameraCalibrationCore/Source/CameraCalibrationCore/Public/LensFile.h
+++ b/Engine/Plugins/VirtualProduction/CameraCalibrationCore/Source/CameraCalibrationCore/Public/LensFile.h
@@ -204,6 +204,10 @@ public:
 	UFUNCTION(BlueprintPure, Category = "Lens Table")
 	bool HasSamples(ELensDataCategory InDataCategory) const;
 
+	/** Returns total number of the points for given category */
+	UFUNCTION(BlueprintPure, Category = "Lens Table")
+	int32 GetTotalPointNum(ELensDataCategory InDataCategory) const;
+
 	/** Get data table reference based on given category */
 	const FBaseLensTable* GetDataTable(ELensDataCategory InDataCategory) const;
 	
diff --git a/Engine/Plugins/VirtualProduction/CameraCalibrationCore/Source/CameraCalibrationCore/Public/Tables/BaseLensTable.h b/Engine/Plugins/VirtualProduction/CameraCalibrationCore/Source/CameraCalibrationCore/Public/Tables/BaseLensTable.h
--- a/Engine/Plugins/VirtualProduction/CameraCalibrationCore/Source/CameraCalibrationCore/Public/Tables/BaseLensTable.h
+++ b/Engine/Plugins/VirtualProduction/CameraCalibrationCore/Source/CameraCalibrationCore/Public/Tables/BaseLensTable.h
@@ -105,6 +105,9 @@ public:
 	/** Get number of Focus points for this data table */
 	virtual int32 GetFocusPointNum() const PURE_VIRTUAL(FBaseLensTable::GetFocusPointNum, return INDEX_NONE; );
 
+	/** Get total number of Zoom points for all Focus points of this data table */
+	virtual int32 GetTotalPointNum() const PURE_VIRTUAL(FBaseLensTable::GetTotalPointNum, return INDEX_NONE; );
+
 	/** Get the base focus point by given index */
 	virtual const FBaseFocusPoint* GetBaseFocusPoint(int32 InIndex) const PURE_VIRTUAL(FBaseLensTable::GetBaseFocusPoint, return nullptr; );
 
diff --git a/Engine/Plugins/VirtualProduction/CameraCalibrationCore/Source/CameraCalibrationCore/Public/Tables/DistortionParametersTable.h b/Engine/Plugins/VirtualProduction/CameraCalibrationCore/Source/CameraCalibrationCore/Public/Tables/DistortionParametersTable.h
--- a/Engine/Plugins/VirtualProduction/CameraCalibrationCore/Source/CameraCalibrationCore/Public/Tables/DistortionParametersTable.h
+++ b/Engine/Plugins/VirtualProduction/CameraCalibrationCore/Source/CameraCalibrationCore/Public/Tables/DistortionParametersTable.h
@@ -95,6 +95,7 @@ public:
 	//~ Begin FBaseDataTable Interface
 	virtual void ForEachPoint(FFocusPointCallback InCallback) const override;
 	virtual int32 GetFocusPointNum() const override { return FocusPoints.Num(); }
+	virtual int32 GetTotalPointNum() const override;
 	virtual UScriptStruct* GetScriptStruct() const override;
 	//~ End FBaseDataTable Interface
 
diff --git a/Engine/Plugins/VirtualProduction/CameraCalibrationCore/Source/CameraCalibrationCore/Public/Tables/FocalLengthTable.h b/Engine/Plugins/VirtualProduction/CameraCalibrationCore/Source/CameraCalibrationCore/Public/Tables/FocalLengthTable.h
--- a/Engine/Plugins/VirtualProduction/CameraCalibrationCore/Source/CameraCalibrationCore/Public/Tables/FocalLengthTable.h
+++ b/Engine/Plugins/VirtualProduction/CameraCalibrationCore/Source/CameraCalibrationCore/Public/Tables/FocalLengthTable.h
@@ -106,6 +106,7 @@ public:
 	//~ Begin FBaseDataTable Interface
 	virtual void ForEachPoint(FFocusPointCallback InCallback) const override;
 	virtual int32 GetFocusPointNum() const override { return FocusPoints.Num(); }
+	virtual int32 GetTotalPointNum() const override;
 	virtual UScriptStruct* GetScriptStruct() const override;
 	//~ End FBaseDataTable Interface
 
diff --git a/Engine/Plugins/VirtualProduction/CameraCalibrationCore/Source/CameraCalibrationCore/Public/Tables/ImageCenterTable.h b/Engine/Plugins/VirtualProduction/CameraCalibrationCore/Source/CameraCalibrationCore/Public/Tables/ImageCenterTable.h
--- a/Engine/Plugins/VirtualProduction/CameraCalibrationCore/Source/CameraCalibrationCore/Public/Tables/ImageCenterTable.h
+++ b/Engine/Plugins/VirtualProduction/CameraCalibrationCore/Source/CameraCalibrationCore/Public/Tables/ImageCenterTable.h
@@ -76,6 +76,7 @@ public:
 	//~ Begin FBaseDataTable Interface
 	virtual void ForEachPoint(FFocusPointCallback InCallback) const override;
 	virtual int32 GetFocusPointNum() const override { return FocusPoints.Num(); }
+	virtual int32 GetTotalPointNum() const override;
 	virtual UScriptStruct* GetScriptStruct() const override;
 	//~ End FBaseDataTable Interface
 	
diff --git a/Engine/Plugins/VirtualProduction/CameraCalibrationCore/Source/CameraCalibrationCore/Public/Tables/NodalOffsetTable.h b/Engine/Plugins/VirtualProduction/CameraCalibrationCore/Source/CameraCalibrationCore/Public/Tables/NodalOffsetTable.h
--- a/Engine/Plugins/VirtualProduction/CameraCalibrationCore/Source/CameraCalibrationCore/Public/Tables/NodalOffsetTable.h
+++ b/Engine/Plugins/VirtualProduction/CameraCalibrationCore/Source/CameraCalibrationCore/Public/Tables/NodalOffsetTable.h
@@ -82,6 +82,7 @@ public:
 	//~ Begin FBaseDataTable Interface
 	virtual void ForEachPoint(FFocusPointCallback InCallback) const override;
 	virtual int32 GetFocusPointNum() const override { return FocusPoints.Num(); }
+	virtual int32 GetTotalPointNum() const override;
 	virtual UScriptStruct* GetScriptStruct() const override;
 	//~ End FBaseDataTable Interface
 
diff --git a/Engine/Plugins/VirtualProduction/CameraCalibrationCore/Source/CameraCalibrationCore/Public/Tables/STMapTable.h b/Engine/Plugins/VirtualProduction/CameraCalibrationCore/Source/CameraCalibrationCore/Public/Tables/STMapTable.h
--- a/Engine/Plugins/VirtualProduction/CameraCalibrationCore/Source/CameraCalibrationCore/Public/Tables/STMapTable.h
+++ b/Engine/Plugins/VirtualProduction/CameraCalibrationCore/Source/CameraCalibrationCore/Public/Tables/STMapTable.h
@@ -134,6 +134,7 @@ public:
 	//~ Begin FBaseDataTable Interface
 	virtual void ForEachPoint(FFocusPointCallback InCallback) const override;
 	virtual int32 GetFocusPointNum() const override { return FocusPoints.Num(); }
+	virtual int32 GetTotalPointNum() const override;
 	virtual UScriptStruct* GetScriptStruct() const override;
 	//~ End FBaseDataTable Interface
 
