[CameraCalibration] Add overlay rendering shader to draw crosshair at an off-center viewport position.

diff --git a/Engine/Plugins/VirtualProduction/CameraCalibrationCore/Shaders/Private/CrosshairOverlay.usf b/Engine/Plugins/VirtualProduction/CameraCalibrationCore/Shaders/Private/CrosshairOverlay.usf
new file mode 100644
--- /dev/null
+++ b/Engine/Plugins/VirtualProduction/CameraCalibrationCore/Shaders/Private/CrosshairOverlay.usf
@@ -0,0 +1,55 @@
+// Copyright Epic Games, Inc. All Rights Reserved.
+
+/*=============================================================================
+	CrosshairOverlay.usf: 
+=============================================================================*/
+
+#include "/Engine/Public/Platform.ush"
+
+// Normalized image center of the lens where the crosshair will be centered
+float2 PrincipalPoint;
+
+// Size of the output texture in (pixels)
+uint2 OutputTextureExtent;
+
+// Length of one side of the crosshair (in pixels)
+float LengthInPixels;
+
+// Size of the gap in the center of the crosshair (in pixels)
+float GapSizeInPixels;
+
+// Color to draw the crosshair
+float4 OutputColor;
+
+bool IsEqualCoordinate(float A, float B)
+{
+	// Pixel coordinates are considered equal if they are within a half-pixel of each other
+	if (abs(A - B) < 0.5)
+	{
+		return true;
+	}
+	return false;
+}
+
+void MainPS(float4 SvPosition : SV_POSITION, out float4 OutColor : SV_Target0)
+{
+	float2 PixelCoordinate = SvPosition.xy - float2(0.5, 0.5); // Half-pixel shift
+    float2 CenterCoordinate = floor(PrincipalPoint * OutputTextureExtent);
+
+	float4 Color = float4(0.0, 0.0, 0.0, 0.0); // Transparent
+    
+    if (IsEqualCoordinate(PixelCoordinate.x, CenterCoordinate.x) && (
+        (PixelCoordinate.y > (CenterCoordinate.y - LengthInPixels) && PixelCoordinate.y < (CenterCoordinate.y - GapSizeInPixels)) || 
+        (PixelCoordinate.y < (CenterCoordinate.y + LengthInPixels)) && (PixelCoordinate.y > (CenterCoordinate.y + GapSizeInPixels))))
+    {
+		Color = OutputColor;
+    }   
+    else if (IsEqualCoordinate(PixelCoordinate.y, CenterCoordinate.y) && (
+        (PixelCoordinate.x > (CenterCoordinate.x - LengthInPixels) && PixelCoordinate.x < (CenterCoordinate.x - GapSizeInPixels)) || 
+        (PixelCoordinate.x < (CenterCoordinate.x + LengthInPixels)) && (PixelCoordinate.x > (CenterCoordinate.x + GapSizeInPixels))))
+    {
+		Color = OutputColor;
+    }
+
+    OutColor = Color;
+}
diff --git a/Engine/Plugins/VirtualProduction/CameraCalibrationCore/Source/CameraCalibrationCore/Private/OverlayRendering.cpp b/Engine/Plugins/VirtualProduction/CameraCalibrationCore/Source/CameraCalibrationCore/Private/OverlayRendering.cpp
new file mode 100644
--- /dev/null
+++ b/Engine/Plugins/VirtualProduction/CameraCalibrationCore/Source/CameraCalibrationCore/Private/OverlayRendering.cpp
@@ -0,0 +1,88 @@
+// Copyright Epic Games, Inc. All Rights Reserved.
+
+#include "OverlayRendering.h"
+
+#include "Engine/TextureRenderTarget2D.h"
+#include "GlobalShader.h"
+#include "Misc/MemStack.h"
+#include "PixelShaderUtils.h"
+#include "RenderGraphUtils.h"
+#include "ScreenPass.h"
+#include "ShaderParameterStruct.h"
+#include "TextureResource.h"
+
+
+/** A pixel shader that draws a crosshair centered at a given principal point */
+class FCrosshairOverlayPS : public FGlobalShader
+{
+	DECLARE_GLOBAL_SHADER(FCrosshairOverlayPS);
+	SHADER_USE_PARAMETER_STRUCT(FCrosshairOverlayPS, FGlobalShader);
+
+	BEGIN_SHADER_PARAMETER_STRUCT(FParameters, )
+
+		SHADER_PARAMETER(FVector2f, PrincipalPoint)
+		SHADER_PARAMETER(FIntPoint, OutputTextureExtent)
+		SHADER_PARAMETER(float, LengthInPixels)
+		SHADER_PARAMETER(float, GapSizeInPixels)
+		SHADER_PARAMETER(FVector4f, OutputColor)
+
+		RENDER_TARGET_BINDING_SLOTS()
+
+	END_SHADER_PARAMETER_STRUCT()
+};
+
+IMPLEMENT_GLOBAL_SHADER(FCrosshairOverlayPS, "/Plugin/CameraCalibrationCore/Private/CrosshairOverlay.usf", "MainPS", SF_Pixel);
+
+
+namespace OverlayRendering
+{
+	bool DrawCrosshairOverlay(UTextureRenderTarget2D* OutRenderTarget, const FCrosshairOverlayParams& CrosshairParams)
+	{
+		if (OutRenderTarget == nullptr)
+		{
+			return false;
+		}
+
+		const FTextureRenderTargetResource* const DestinationTextureResource = OutRenderTarget->GameThread_GetRenderTargetResource();
+
+		if (!DestinationTextureResource)
+		{
+			return false;
+		}
+
+		ENQUEUE_RENDER_COMMAND(OverlayRendering_CrosshairOverlay)(
+			[DestinationTextureResource, CrosshairParams](FRHICommandListImmediate& RHICmdList)
+			{
+				FMemMark Mark(FMemStack::Get());
+				FRDGBuilder GraphBuilder(RHICmdList);
+
+				FCrosshairOverlayPS::FParameters* PassParameters = GraphBuilder.AllocParameters<FCrosshairOverlayPS::FParameters>();
+
+				FRDGTextureRef OutputTexture = GraphBuilder.RegisterExternalTexture(CreateRenderTarget(DestinationTextureResource->TextureRHI, TEXT("OutputOverlay")));
+				PassParameters->OutputTextureExtent = OutputTexture->Desc.Extent;
+				PassParameters->PrincipalPoint = CrosshairParams.PrincipalPoint;
+				PassParameters->LengthInPixels = CrosshairParams.LengthInPixels;
+				PassParameters->GapSizeInPixels = CrosshairParams.GapSizeInPixels;
+				PassParameters->OutputColor = CrosshairParams.CrosshairColor;
+
+				FGlobalShaderMap* GlobalShaderMap = GetGlobalShaderMap(GMaxRHIFeatureLevel);
+				const TShaderMapRef<FCrosshairOverlayPS> PixelShader(GlobalShaderMap);
+
+				FScreenPassRenderTarget SceneColorRenderTarget(OutputTexture, ERenderTargetLoadAction::ELoad);
+				PassParameters->RenderTargets[0] = SceneColorRenderTarget.GetRenderTargetBinding();
+
+				FPixelShaderUtils::AddFullscreenPass(
+					GraphBuilder
+					, GlobalShaderMap
+					, RDG_EVENT_NAME("CrosshairOverlay")
+					, PixelShader
+					, PassParameters
+					, FIntRect(0, 0, OutputTexture->Desc.Extent.X, OutputTexture->Desc.Extent.Y));
+
+				GraphBuilder.Execute();
+			});
+
+		return true;
+	}
+}
+
diff --git a/Engine/Plugins/VirtualProduction/CameraCalibrationCore/Source/CameraCalibrationCore/Public/OverlayRendering.h b/Engine/Plugins/VirtualProduction/CameraCalibrationCore/Source/CameraCalibrationCore/Public/OverlayRendering.h
new file mode 100644
--- /dev/null
+++ b/Engine/Plugins/VirtualProduction/CameraCalibrationCore/Source/CameraCalibrationCore/Public/OverlayRendering.h
@@ -0,0 +1,32 @@
+// Copyright Epic Games, Inc. All Rights Reserved.
+
+#pragma once
+
+#include "CoreTypes.h"
+#include "Math/Vector2D.h"
+#include "Math/Vector4.h"
+
+class UTextureRenderTarget2D;
+
+/** Parameters for the crosshair shader */
+struct FCrosshairOverlayParams
+{
+	/** Normalized image center of the lens where the crosshair will be centered */
+	FVector2f PrincipalPoint = { 0.5f, 0.5f };
+
+	/** Length of one side of the crosshair (in pixels) */
+	float LengthInPixels = 50.0f;
+
+	/** Size of the gap in the center of the crosshair (in pixels) */
+	float GapSizeInPixels = 2.0f;
+
+	/** Color to draw the crosshair */
+	FVector4f CrosshairColor = FVector4f(1.0f, 1.0f, 1.0f, 1.0f);
+};
+
+
+namespace OverlayRendering
+{
+	/** Draws a crosshair overlay to the output render target */
+	CAMERACALIBRATIONCORE_API bool DrawCrosshairOverlay(UTextureRenderTarget2D* OutRenderTarget, const FCrosshairOverlayParams& CrosshairParams);
+}
