Properly initialize ImageDimensions script struct member

diff --git a/Engine/Plugins/VirtualProduction/CameraCalibrationCore/Source/CameraCalibrationCore/Public/LensData.h b/Engine/Plugins/VirtualProduction/CameraCalibrationCore/Source/CameraCalibrationCore/Public/LensData.h
--- a/Engine/Plugins/VirtualProduction/CameraCalibrationCore/Source/CameraCalibrationCore/Public/LensData.h
+++ b/Engine/Plugins/VirtualProduction/CameraCalibrationCore/Source/CameraCalibrationCore/Public/LensData.h
@@ -51,7 +51,7 @@ public:
 
 	/** Camera image size, in pixels */
 	UPROPERTY(EditAnywhere, BlueprintReadWrite, Category = "Lens Info")
-	FIntPoint ImageDimensions = FIntPoint(0, 1);
+	FIntPoint ImageDimensions = FIntPoint(1920, 1080);
 
 };
 
