Fix camera calibration producing NANs on Vulkan

diff --git a/Engine/Plugins/VirtualProduction/CameraCalibrationCore/Source/CameraCalibrationCore/Private/DistortionRenderingUtils.cpp b/Engine/Plugins/VirtualProduction/CameraCalibrationCore/Source/CameraCalibrationCore/Private/DistortionRenderingUtils.cpp
--- a/Engine/Plugins/VirtualProduction/CameraCalibrationCore/Source/CameraCalibrationCore/Private/DistortionRenderingUtils.cpp
+++ b/Engine/Plugins/VirtualProduction/CameraCalibrationCore/Source/CameraCalibrationCore/Private/DistortionRenderingUtils.cpp
@@ -86,6 +86,7 @@ namespace DistortionRenderingUtils
 				RHICmdList.CopyToStagingBuffer(UndistortedPointsBuffer.Buffer, DestinationStagingBuffer, 0, BufferSize);
 
 				// Wait to ensure that the staging buffer is ready to read
+				RHICmdList.SubmitCommandsAndFlushGPU();
 				RHICmdList.BlockUntilGPUIdle();
 
 				// Copy data out of the staging buffer
