Rename LensDistortionComponent to LensComponent and move from CameraCalibration plugin to CameraCalibrationCore plugin
Remove CameraCalibration module from CameraCalibration plugin

diff --git a/Engine/Plugins/VirtualProduction/CameraCalibration/CameraCalibration.uplugin b/Engine/Plugins/VirtualProduction/CameraCalibration/CameraCalibration.uplugin
--- a/Engine/Plugins/VirtualProduction/CameraCalibration/CameraCalibration.uplugin
+++ b/Engine/Plugins/VirtualProduction/CameraCalibration/CameraCalibration.uplugin
@@ -14,11 +14,6 @@
 	"IsBetaVersion": true,
 	"Installed": false,
 	"Modules": [
-		{
-			"Name": "CameraCalibration",
-			"Type": "Runtime",
-			"LoadingPhase": "Default"
-		},
 		{
 			"Name": "CameraCalibrationEditor",
 			"Type": "Editor",
diff --git a/Engine/Plugins/VirtualProduction/CameraCalibration/Source/CameraCalibration/CameraCalibration.Build.cs b/Engine/Plugins/VirtualProduction/CameraCalibration/Source/CameraCalibration/CameraCalibration.Build.cs
deleted file mode 100644
diff --git a/Engine/Plugins/VirtualProduction/CameraCalibration/Source/CameraCalibration/Private/CameraCalibrationLog.h b/Engine/Plugins/VirtualProduction/CameraCalibration/Source/CameraCalibration/Private/CameraCalibrationLog.h
deleted file mode 100644
diff --git a/Engine/Plugins/VirtualProduction/CameraCalibration/Source/CameraCalibration/Private/CameraCalibrationModule.cpp b/Engine/Plugins/VirtualProduction/CameraCalibration/Source/CameraCalibration/Private/CameraCalibrationModule.cpp
deleted file mode 100644
diff --git a/Engine/Plugins/VirtualProduction/CameraCalibration/Source/CameraCalibration/Private/CameraCalibrationModule.h b/Engine/Plugins/VirtualProduction/CameraCalibration/Source/CameraCalibration/Private/CameraCalibrationModule.h
deleted file mode 100644
diff --git a/Engine/Plugins/VirtualProduction/CameraCalibration/Source/CameraCalibration/Private/LensDistortionComponent.cpp b/Engine/Plugins/VirtualProduction/CameraCalibration/Source/CameraCalibration/Private/LensDistortionComponent.cpp
deleted file mode 100644
diff --git a/Engine/Plugins/VirtualProduction/CameraCalibration/Source/CameraCalibration/Public/LensDistortionComponent.h b/Engine/Plugins/VirtualProduction/CameraCalibration/Source/CameraCalibration/Public/LensDistortionComponent.h
deleted file mode 100644
diff --git a/Engine/Plugins/VirtualProduction/CameraCalibration/Source/CameraCalibrationEditor/CameraCalibrationEditor.Build.cs b/Engine/Plugins/VirtualProduction/CameraCalibration/Source/CameraCalibrationEditor/CameraCalibrationEditor.Build.cs
--- a/Engine/Plugins/VirtualProduction/CameraCalibration/Source/CameraCalibrationEditor/CameraCalibrationEditor.Build.cs
+++ b/Engine/Plugins/VirtualProduction/CameraCalibration/Source/CameraCalibrationEditor/CameraCalibrationEditor.Build.cs
@@ -13,7 +13,6 @@ namespace UnrealBuildTool.Rules
 					"AppFramework",
 					"AssetRegistry",
 					"AssetTools",
-					"CameraCalibration",
 					"CameraCalibrationCore",
 					"CinematicCamera",
 					"Composure",
diff --git a/Engine/Plugins/VirtualProduction/CameraCalibrationCore/CameraCalibrationCore.uplugin b/Engine/Plugins/VirtualProduction/CameraCalibrationCore/CameraCalibrationCore.uplugin
--- a/Engine/Plugins/VirtualProduction/CameraCalibrationCore/CameraCalibrationCore.uplugin
+++ b/Engine/Plugins/VirtualProduction/CameraCalibrationCore/CameraCalibrationCore.uplugin
@@ -30,6 +30,10 @@
 		{
 			"Name": "ProceduralMeshComponent",
 			"Enabled": true
+		},
+		{
+			"Name": "LiveLink",
+			"Enabled": true
 		}
 	]
 }
\ No newline at end of file
diff --git a/Engine/Plugins/VirtualProduction/CameraCalibrationCore/Config/BaseCameraCalibrationCore.ini b/Engine/Plugins/VirtualProduction/CameraCalibrationCore/Config/BaseCameraCalibrationCore.ini
--- a/Engine/Plugins/VirtualProduction/CameraCalibrationCore/Config/BaseCameraCalibrationCore.ini
+++ b/Engine/Plugins/VirtualProduction/CameraCalibrationCore/Config/BaseCameraCalibrationCore.ini
@@ -25,3 +25,15 @@
 
 +PackageRedirects=(OldName="/CameraCalibration/Materials/M_SphericalDistortionPostProcess", NewName="/CameraCalibrationCore/Materials/M_DistortionPostProcess")
 +PackageRedirects=(OldName="/CameraCalibration/Materials/M_SphericalDistortionDisplacementMap", NewName="/CameraCalibrationCore/Materials/M_SphericalDistortionDisplacementMap")
+
+;
+; 5.1 Redirects
+;
+
+; Classes
+
++ClassRedirects=(OldName="LensDistortionComponent", NewName="/Script/CameraCalibrationCore.LensComponent")
+
+; Packages
+
++PackageRedirects=(OldName="/Script/CameraCalibration", Removed=true)
diff --git a/Engine/Plugins/VirtualProduction/CameraCalibrationCore/Source/CameraCalibrationCore/CameraCalibrationCore.Build.cs b/Engine/Plugins/VirtualProduction/CameraCalibrationCore/Source/CameraCalibrationCore/CameraCalibrationCore.Build.cs
--- a/Engine/Plugins/VirtualProduction/CameraCalibrationCore/Source/CameraCalibrationCore/CameraCalibrationCore.Build.cs
+++ b/Engine/Plugins/VirtualProduction/CameraCalibrationCore/Source/CameraCalibrationCore/CameraCalibrationCore.Build.cs
@@ -25,6 +25,7 @@ namespace UnrealBuildTool.Rules
 			PrivateDependencyModuleNames.AddRange(
 				new string[]
                 {
+					"LiveLinkComponents",
                     "Projects",
 				}
 			);
diff --git a/Engine/Plugins/VirtualProduction/CameraCalibrationCore/Source/CameraCalibrationCore/Private/LensComponent.cpp b/Engine/Plugins/VirtualProduction/CameraCalibrationCore/Source/CameraCalibrationCore/Private/LensComponent.cpp
new file mode 100644
--- /dev/null
+++ b/Engine/Plugins/VirtualProduction/CameraCalibrationCore/Source/CameraCalibrationCore/Private/LensComponent.cpp
@@ -0,0 +1,467 @@
+// Copyright Epic Games, Inc. All Rights Reserved.
+
+#include "LensComponent.h"
+
+#include "CameraCalibrationSubsystem.h"
+#include "CineCameraComponent.h"
+#include "Engine/Engine.h"
+#include "LiveLinkComponentController.h"
+
+ULensComponent::ULensComponent()
+{
+	PrimaryComponentTick.bCanEverTick = true;
+	PrimaryComponentTick.bStartWithTickEnabled = true;
+	PrimaryComponentTick.TickGroup = ETickingGroup::TG_LastDemotable;
+	bTickInEditor = true;
+
+	if (!HasAnyFlags(RF_ArchetypeObject | RF_ClassDefaultObject))
+	{
+		//Hook up to PostActorTick to handle nodal offset
+		FWorldDelegates::OnWorldPostActorTick.AddUObject(this, &ULensComponent::OnPostActorTick);
+
+		DistortionProducerID = FGuid::NewGuid();
+	}
+}
+
+void ULensComponent::OnRegister()
+{
+	Super::OnRegister();
+
+	InitDefaultCamera();
+
+	if (UCineCameraComponent* const CineCameraComponent = Cast<UCineCameraComponent>(TargetCameraComponent.GetComponent(GetOwner())))
+	{
+		LastRotation = CineCameraComponent->GetRelativeRotation();
+		LastLocation = CineCameraComponent->GetRelativeLocation();
+	}
+}
+
+void ULensComponent::DestroyComponent(bool bPromoteChildren /*= false*/)
+{
+	Super::DestroyComponent(bPromoteChildren);
+
+	if (UCineCameraComponent* const CineCameraComponent = Cast<UCineCameraComponent>(TargetCameraComponent.GetComponent(GetOwner())))
+	{
+		CleanupDistortion(CineCameraComponent);
+
+		if (UCameraCalibrationSubsystem* SubSystem = GEngine->GetEngineSubsystem<UCameraCalibrationSubsystem>())
+		{
+			SubSystem->UnregisterDistortionModelHandler(CineCameraComponent, ProducedLensDistortionHandler);
+		}
+
+		// Restore the target camera's location and rotation (undo nodal offset)
+		CineCameraComponent->SetRelativeLocation(OriginalCameraLocation);
+		CineCameraComponent->SetRelativeRotation(OriginalCameraRotation.Quaternion());
+	}
+}
+
+void ULensComponent::TickComponent(float DeltaTime, ELevelTick TickType, FActorComponentTickFunction* ThisTickFunction)
+{
+	if (UCineCameraComponent* const CineCameraComponent = Cast<UCineCameraComponent>(TargetCameraComponent.GetComponent(GetOwner())))
+	{
+		// Add a tick prerequisite on any live link components also on this camera
+		// LiveLink might set the focal length of the target camera, so it is important that this component ticks after that
+		TInlineComponentArray<ULiveLinkComponentController*> LiveLinkComponents;
+		CineCameraComponent->GetOwner()->GetComponents(LiveLinkComponents);
+
+		for (UActorComponent* LiveLinkComponent : LiveLinkComponents)
+		{
+			AddTickPrerequisiteComponent(LiveLinkComponent);
+		}
+
+		DistortionSource.TargetCameraComponent = CineCameraComponent;
+
+		UCameraCalibrationSubsystem* const SubSystem = GEngine->GetEngineSubsystem<UCameraCalibrationSubsystem>();
+
+		// Evaluate the selected lens file for distortion data using the current settings of the target camera
+		ULensFile* LensFile = LensFilePicker.GetLensFile();
+		if (bEvaluateLensFileForDistortion && LensFile)
+		{
+			if (!ProducedLensDistortionHandler)
+			{
+				CreateDistortionHandlerForLensFile();
+			}
+
+			//Go through the lens file to get distortion data based on FIZ
+			//Our handler's displacement map will get updated
+			const FVector2D CurrentSensorDimensions = FVector2D(CineCameraComponent->Filmback.SensorWidth, CineCameraComponent->Filmback.SensorHeight);
+
+			if (SubSystem)
+			{
+				SubSystem->GetOriginalFocalLength(CineCameraComponent, OriginalFocalLength);
+			}
+
+			LensFile->EvaluateDistortionData(CineCameraComponent->CurrentFocusDistance, OriginalFocalLength, CurrentSensorDimensions, ProducedLensDistortionHandler);
+
+			// Adjust overscan by the overscan multiplier
+			if (bScaleOverscan)
+			{
+				const float ScaledOverscanFactor = ((ProducedLensDistortionHandler->GetOverscanFactor() - 1.0f) * OverscanMultiplier) + 1.0f;
+				ProducedLensDistortionHandler->SetOverscanFactor(ScaledOverscanFactor);
+			}
+
+			// Track changes to the cine camera's focal length for consumers of distortion data
+			if (SubSystem)
+			{
+				SubSystem->UpdateOriginalFocalLength(CineCameraComponent, CineCameraComponent->CurrentFocalLength);
+			}
+		}
+	
+		if (bApplyDistortion)
+		{
+			// Get the lens distortion handler for the target camera and distortion source
+			ULensDistortionModelHandlerBase* LensDistortionHandler = nullptr;
+			if (SubSystem)
+			{
+				LensDistortionHandler = SubSystem->FindDistortionModelHandler(DistortionSource);
+			}
+
+			if (LensDistortionHandler)
+			{
+				// Get the current distortion MID from the lens distortion handler
+				UMaterialInstanceDynamic* NewDistortionMID = LensDistortionHandler->GetDistortionMID();
+
+				// If the MID has changed
+				if (LastDistortionMID != NewDistortionMID)
+				{
+					CineCameraComponent->RemoveBlendable(LastDistortionMID);
+					CineCameraComponent->AddOrUpdateBlendable(NewDistortionMID);
+				}
+
+				// Cache the latest distortion MID
+				LastDistortionMID = NewDistortionMID;
+
+				SubSystem->GetOriginalFocalLength(CineCameraComponent, OriginalFocalLength);
+
+				// Get the overscan factor and use it to modify the target camera's FOV
+				const float OverscanFactor = LensDistortionHandler->GetOverscanFactor();
+				const float OverscanSensorWidth = CineCameraComponent->Filmback.SensorWidth * OverscanFactor;
+				const float OverscanFOV = FMath::RadiansToDegrees(2.0f * FMath::Atan(OverscanSensorWidth / (2.0f * OriginalFocalLength)));
+				CineCameraComponent->SetFieldOfView(OverscanFOV);
+
+				// Update the minimum and maximum focal length of the camera (if needed)
+				CineCameraComponent->LensSettings.MinFocalLength = FMath::Min(CineCameraComponent->LensSettings.MinFocalLength, CineCameraComponent->CurrentFocalLength);
+				CineCameraComponent->LensSettings.MaxFocalLength = FMath::Max(CineCameraComponent->LensSettings.MaxFocalLength, CineCameraComponent->CurrentFocalLength);
+
+				SubSystem->UpdateOverscanFocalLength(CineCameraComponent, CineCameraComponent->CurrentFocalLength);
+
+				bIsDistortionSetup = true;
+			}
+			else
+			{
+				CleanupDistortion(CineCameraComponent);
+			}
+		}
+	}
+
+	Super::TickComponent(DeltaTime, TickType, ThisTickFunction);
+}
+
+#if WITH_EDITOR
+void ULensComponent::PostEditChangeProperty(struct FPropertyChangedEvent& PropertyChangedEvent)
+{
+ 	const FName PropertyName = (PropertyChangedEvent.Property != NULL) ? PropertyChangedEvent.Property->GetFName() : NAME_None;
+ 
+ 	if (PropertyName == GET_MEMBER_NAME_CHECKED(ULensComponent, bApplyDistortion))
+ 	{
+ 		if (UCineCameraComponent* const CineCameraComponent = Cast<UCineCameraComponent>(TargetCameraComponent.GetComponent(GetOwner())))
+ 		{
+ 			if (!bApplyDistortion)
+ 			{
+				CleanupDistortion(CineCameraComponent);
+ 			}
+ 		}
+ 	}
+	else if (PropertyName == GET_MEMBER_NAME_CHECKED(ULensComponent, bApplyNodalOffset))
+	{
+		if (UCineCameraComponent* const CineCameraComponent = Cast<UCineCameraComponent>(TargetCameraComponent.GetComponent(GetOwner())))
+		{
+			if (bApplyNodalOffset)
+			{
+				OriginalCameraRotation = CineCameraComponent->GetRelativeRotation();
+				OriginalCameraLocation = CineCameraComponent->GetRelativeLocation();
+			}
+			else
+			{
+				CineCameraComponent->SetRelativeLocation(OriginalCameraLocation);
+				CineCameraComponent->SetRelativeRotation(OriginalCameraRotation.Quaternion());
+			}
+		}
+	}
+	else if (PropertyName == GET_MEMBER_NAME_CHECKED(ULensComponent, bEvaluateLensFileForDistortion))
+	{
+		if (bEvaluateLensFileForDistortion)
+		{
+			// Create a new distortion handler if one does not yet exist for the selected lens file
+			if (!ProducedLensDistortionHandler)
+			{
+				CreateDistortionHandlerForLensFile();
+			}
+		}
+		else
+		{
+			// Reset the distortion source settings
+			if (DistortionSource.DistortionProducerID == DistortionProducerID)
+			{
+				DistortionSource.HandlerDisplayName.Empty();
+				DistortionSource.DistortionProducerID.Invalidate();
+			}
+
+			// Unregister the produced handler from the subsystem
+			if (UCameraCalibrationSubsystem* SubSystem = GEngine->GetEngineSubsystem<UCameraCalibrationSubsystem>())
+			{
+				if (UCineCameraComponent* const CineCameraComponent = Cast<UCineCameraComponent>(TargetCameraComponent.GetComponent(GetOwner())))
+				{
+					SubSystem->UnregisterDistortionModelHandler(CineCameraComponent, ProducedLensDistortionHandler);
+				}
+			}
+
+			ProducedLensDistortionHandler = nullptr;
+
+			// Restore the target camera's location and rotation (undo nodal offset)
+			if (UCineCameraComponent* const CineCameraComponent = Cast<UCineCameraComponent>(TargetCameraComponent.GetComponent(GetOwner())))
+			{
+				CineCameraComponent->SetRelativeLocation(OriginalCameraLocation);
+				CineCameraComponent->SetRelativeRotation(OriginalCameraRotation.Quaternion());
+			}
+		}
+	}
+	else if (PropertyName == GET_MEMBER_NAME_CHECKED(ULensComponent, TargetCameraComponent))
+	{
+		// Clean up distortion on the last target camera
+		if (LastCameraComponent)
+		{
+			CleanupDistortion(LastCameraComponent);
+
+			// Unregister the produced handler (if one exists) from the subsystem for the last camera component
+			if (ProducedLensDistortionHandler)
+			{
+				if (UCameraCalibrationSubsystem* SubSystem = GEngine->GetEngineSubsystem<UCameraCalibrationSubsystem>())
+				{
+					SubSystem->UnregisterDistortionModelHandler(LastCameraComponent, ProducedLensDistortionHandler);
+				}
+
+				ProducedLensDistortionHandler = nullptr;
+			}
+
+			// Restore the last camera's location and rotation (undo nodal offset)
+			LastCameraComponent->SetRelativeLocation(OriginalCameraLocation);
+			LastCameraComponent->SetRelativeRotation(OriginalCameraRotation.Quaternion());
+		}
+
+		LastCameraComponent = Cast<UCineCameraComponent>(TargetCameraComponent.GetComponent(GetOwner()));
+		DistortionSource.TargetCameraComponent = LastCameraComponent;
+
+		// Create a new handler for the new target camera component
+		if (bEvaluateLensFileForDistortion)
+		{
+			CreateDistortionHandlerForLensFile();
+		}
+
+		// Cache the rotation and location of the new target camera component
+		if (UCineCameraComponent* const CineCameraComponent = Cast<UCineCameraComponent>(TargetCameraComponent.GetComponent(GetOwner())))
+		{
+			LastRotation = OriginalCameraRotation = CineCameraComponent->GetRelativeRotation();
+			LastLocation = OriginalCameraLocation = CineCameraComponent->GetRelativeLocation();
+		}
+	}
+	else if (PropertyName == GET_MEMBER_NAME_CHECKED(FLensFilePicker, LensFile))
+	{
+		if (bEvaluateLensFileForDistortion)
+		{
+			CreateDistortionHandlerForLensFile();
+		}
+	}
+
+	Super::PostEditChangeProperty(PropertyChangedEvent);
+}
+#endif //WITH_EDITOR
+
+void ULensComponent::PostDuplicate(bool bDuplicateForPIE)
+{
+	Super::PostDuplicate(bDuplicateForPIE);
+
+	InitDefaultCamera();
+
+	// When this component is duplicated (e.g. for PIE), the duplicated component needs a new unique ID
+	if (!DistortionProducerID.IsValid())
+	{
+		DistortionProducerID = FGuid::NewGuid();
+	}
+}
+
+void ULensComponent::PostEditImport()
+{
+	Super::PostEditImport();
+
+	InitDefaultCamera();
+
+	// When this component is duplicated (e.g. copy-paste), the duplicated component needs a new unique ID
+	if (!DistortionProducerID.IsValid())
+	{
+		DistortionProducerID = FGuid::NewGuid();
+	}
+}
+
+void ULensComponent::PostLoad()
+{
+	Super::PostLoad();
+
+	if (bIsDistortionSetup)
+	{
+		if (UCineCameraComponent* const CineCameraComponent = Cast<UCineCameraComponent>(TargetCameraComponent.GetComponent(GetOwner())))
+		{
+			CleanupDistortion(CineCameraComponent);
+		}
+		bIsDistortionSetup = false;
+	}
+}
+
+void ULensComponent::OnPostActorTick(UWorld* World, ELevelTick TickType, float DeltaSeconds)
+{
+	if (World == GetWorld())
+	{
+		if (UCineCameraComponent* const CineCameraComponent = Cast<UCineCameraComponent>(TargetCameraComponent.GetComponent(GetOwner())))
+		{
+			if (ULensFile* CurrentLensFile = LensFilePicker.GetLensFile())
+			{
+				ApplyNodalOffset(CurrentLensFile, CineCameraComponent);
+			}
+		}
+	}
+}
+
+void ULensComponent::CleanupDistortion(UCineCameraComponent* const CineCameraComponent)
+{
+	if (bIsDistortionSetup)
+	{
+		// Remove the last distortion MID that was applied to the target camera component
+		if (LastDistortionMID)
+		{
+			CineCameraComponent->RemoveBlendable(LastDistortionMID);
+			LastDistortionMID = nullptr;
+		}
+
+		// Restore the original FOV of the target camera
+		const float UndistortedFOV = FMath::RadiansToDegrees(2.0f * FMath::Atan(CineCameraComponent->Filmback.SensorWidth / (2.0f * OriginalFocalLength)));
+		CineCameraComponent->SetFieldOfView(UndistortedFOV);
+
+		// Update the minimum and maximum focal length of the camera (if needed)
+		CineCameraComponent->LensSettings.MinFocalLength = FMath::Min(CineCameraComponent->LensSettings.MinFocalLength, CineCameraComponent->CurrentFocalLength);
+		CineCameraComponent->LensSettings.MaxFocalLength = FMath::Max(CineCameraComponent->LensSettings.MaxFocalLength, CineCameraComponent->CurrentFocalLength);
+
+		UCameraCalibrationSubsystem* SubSystem = GEngine->GetEngineSubsystem<UCameraCalibrationSubsystem>();
+		if (SubSystem)
+		{
+			SubSystem->UpdateOverscanFocalLength(CineCameraComponent, 0.0f);
+		}
+	}
+
+	bIsDistortionSetup = false;
+}
+
+void ULensComponent::InitDefaultCamera()
+{
+	UCineCameraComponent* const CineCameraComponent = Cast<UCineCameraComponent>(TargetCameraComponent.GetComponent(GetOwner()));
+	if (!CineCameraComponent)
+	{
+		// Find all CineCameraComponents on the same actor as this component and set the first one to be the target
+		TInlineComponentArray<UCineCameraComponent*> CineCameraComponents;
+		GetOwner()->GetComponents(CineCameraComponents);
+		if (CineCameraComponents.Num() > 0)
+		{
+			TargetCameraComponent.ComponentProperty = CineCameraComponents[0]->GetFName();
+			LastCameraComponent = CineCameraComponents[0];
+
+			OriginalCameraRotation = LastCameraComponent->GetRelativeRotation();
+			OriginalCameraLocation = LastCameraComponent->GetRelativeLocation();
+		}
+	}
+}
+
+void ULensComponent::CreateDistortionHandlerForLensFile()
+{
+	if (DistortionSource.TargetCameraComponent == nullptr)
+	{
+		return;
+	}
+
+	ULensFile* const LensFile = LensFilePicker.GetLensFile();
+
+	// Set the distortion source settings to match this component and the selected lens file
+	if (LensFile)
+	{
+		DistortionSource.HandlerDisplayName = FString::Format(TEXT("{0} (Lens File) (Lens Distortion Component)"), { LensFile->GetFName().ToString() });
+		DistortionSource.DistortionProducerID = DistortionProducerID;
+
+		UCameraCalibrationSubsystem* const SubSystem = GEngine->GetEngineSubsystem<UCameraCalibrationSubsystem>();
+		if (SubSystem)
+		{
+			ProducedLensDistortionHandler = SubSystem->FindOrCreateDistortionModelHandler(DistortionSource, LensFile->LensInfo.LensModel);
+		}
+	}
+}
+
+void ULensComponent::ApplyNodalOffset(ULensFile* SelectedLensFile, UCineCameraComponent* CineCameraComponent)
+{
+	if (CineCameraComponent)
+	{
+		const FRotator CurrentRotator = CineCameraComponent->GetRelativeRotation();
+		const FVector CurrentTranslation = CineCameraComponent->GetRelativeLocation();
+
+		//Verify if something was set by user / programmatically and update original values
+		if (CurrentRotator != LastRotation)
+		{
+			if (CurrentRotator.Pitch != LastRotation.Pitch)
+			{
+				OriginalCameraRotation.Pitch = CurrentRotator.Pitch;
+			}
+
+			if (CurrentRotator.Yaw != LastRotation.Yaw)
+			{
+				OriginalCameraRotation.Yaw = CurrentRotator.Yaw;
+			}
+
+			if (CurrentRotator.Roll != LastRotation.Roll)
+			{
+				OriginalCameraRotation.Roll = CurrentRotator.Roll;
+			}
+		}
+
+		if (CurrentTranslation != LastLocation)
+		{
+			if (CurrentTranslation.X != LastLocation.X)
+			{
+				OriginalCameraLocation.X = CurrentTranslation.X;
+			}
+			if (CurrentTranslation.Y != LastLocation.Y)
+			{
+				OriginalCameraLocation.Y = CurrentTranslation.Y;
+			}
+			if (CurrentTranslation.Z != LastLocation.Z)
+			{
+				OriginalCameraLocation.Z = CurrentTranslation.Z;
+			}
+		}
+
+		if (bEvaluateLensFileForDistortion && bApplyNodalOffset && SelectedLensFile)
+		{
+			FNodalPointOffset Offset;
+
+			SelectedLensFile->EvaluateNodalPointOffset(CineCameraComponent->CurrentFocusDistance, CineCameraComponent->CurrentFocalLength, Offset);
+
+			CineCameraComponent->SetRelativeLocation(OriginalCameraLocation);
+			CineCameraComponent->SetRelativeRotation(OriginalCameraRotation.Quaternion());
+			CineCameraComponent->AddLocalOffset(Offset.LocationOffset);
+			CineCameraComponent->AddLocalRotation(Offset.RotationOffset);
+		}
+
+		LastLocation = CineCameraComponent->GetRelativeLocation();
+		LastRotation = CineCameraComponent->GetRelativeRotation();
+	}
+	else
+	{
+		LastLocation = FVector::OneVector;
+		LastRotation = FRotator::ZeroRotator;
+	}
+}
diff --git a/Engine/Plugins/VirtualProduction/CameraCalibrationCore/Source/CameraCalibrationCore/Public/LensComponent.h b/Engine/Plugins/VirtualProduction/CameraCalibrationCore/Source/CameraCalibrationCore/Public/LensComponent.h
new file mode 100644
--- /dev/null
+++ b/Engine/Plugins/VirtualProduction/CameraCalibrationCore/Source/CameraCalibrationCore/Public/LensComponent.h
@@ -0,0 +1,129 @@
+// Copyright Epic Games, Inc. All Rights Reserved.
+
+#pragma once
+
+#include "Components/ActorComponent.h"
+
+#include "CameraCalibrationTypes.h"
+#include "LensFile.h"
+
+#include "LensComponent.generated.h"
+
+class UCineCameraComponent;
+
+/** Component for applying a post-process lens distortion effect to a CineCameraComponent on the same actor */
+UCLASS(HideCategories = (Tags, Activation, Cooking, AssetUserData, Collision), meta = (BlueprintSpawnableComponent))
+class CAMERACALIBRATIONCORE_API ULensComponent : public UActorComponent
+{
+	GENERATED_BODY()
+
+public:
+	ULensComponent();
+
+	//~ Begin UActorComponent interface
+	virtual void OnRegister() override;
+	virtual void TickComponent(float DeltaTime, ELevelTick TickType, FActorComponentTickFunction* ThisTickFunction) override;
+	virtual void DestroyComponent(bool bPromoteChildren = false) override;
+	//~ End UActorComponent interface
+
+	//~ UObject interface
+#if WITH_EDITOR
+	virtual void PostEditChangeProperty(struct FPropertyChangedEvent& PropertyChangedEvent) override;
+#endif //WITH_EDITOR
+	virtual void PostDuplicate(bool bDuplicateForPIE) override;
+	virtual void PostEditImport() override;
+	virtual void PostLoad() override;
+	//~ End UObject interface
+
+protected:
+	/** Verify base transform and apply nodal offset on top of everything else done in tick */
+	void OnPostActorTick(UWorld* World, ELevelTick TickType, float DeltaSeconds);
+
+private:
+	/** If TargetCameraComponent is not set, initialize it to the first CineCameraComponent on the same actor as this component */
+	void InitDefaultCamera();
+
+	/** Remove the last distortion MID applied to the input CineCameraComponent and reset its FOV to use no overscan */
+	void CleanupDistortion(UCineCameraComponent* CineCameraComponent);
+
+	/** Register a new lens distortion handler with the camera calibration subsystem using the selected lens file */
+	void CreateDistortionHandlerForLensFile();
+
+	/** Evaluate the input LensFile and apply the nodal offset to the input camera's transform */
+	void ApplyNodalOffset(ULensFile* SelectedLensFile, UCineCameraComponent* CineCameraComponent);
+
+protected:
+	/** The CineCameraComponent on which to apply the post-process distortion effect */
+	UPROPERTY(EditInstanceOnly, Category = "Default", meta = (UseComponentPicker, AllowedClasses = "CineCameraComponent"))
+	FComponentReference TargetCameraComponent;
+
+	/** Structure used to query the camera calibration subsystem for a lens distortion model handler */
+	UPROPERTY(EditAnywhere, BlueprintReadWrite, Category = "Default")
+	FDistortionHandlerPicker DistortionSource;
+
+	/** Whether or not to apply distortion to the target camera component */
+	UPROPERTY(EditAnywhere, BlueprintReadWrite, Category = "Default")
+	bool bApplyDistortion = false;
+
+	/** Whether a distortion effect is currently being applied to the target camera component */
+	UPROPERTY()
+	bool bIsDistortionSetup = false;
+
+	/** Focal length of the target camera before any overscan has been applied */
+	UPROPERTY()
+	float OriginalFocalLength = 35.0f;
+
+	/** Original camera rotation to reset before/after applying nodal offset */
+	UPROPERTY()
+	FRotator OriginalCameraRotation;
+
+	/** Original camera location to reset before/after applying nodal offset */
+	UPROPERTY()
+	FVector OriginalCameraLocation;
+
+	/** Cached MID last applied to the target camera */
+	UPROPERTY()
+	UMaterialInstanceDynamic* LastDistortionMID = nullptr;
+
+	/** Cached most recent target camera, used to clean up the old camera when the user changes the target */
+	UPROPERTY(Transient)
+	UCineCameraComponent* LastCameraComponent = nullptr;
+
+	/** 
+	 * Whether to use the specified lens file to drive distortion
+	 * Enabling this will create a new distortion source and automatically set this component to use it
+	 */
+	UPROPERTY(EditAnywhere, AdvancedDisplay, Category = "Default")
+	bool bEvaluateLensFileForDistortion = false;
+
+	/** Whether to apply nodal offset to the target camera */
+	UPROPERTY(EditAnywhere, AdvancedDisplay, Category = "Default", meta = (EditCondition = "bEvaluateLensFileForDistortion"))
+	bool bApplyNodalOffset = true;
+
+	/** Lens File used to drive distortion with current camera settings */
+	UPROPERTY(EditAnywhere, AdvancedDisplay, Category = "Default", meta = (EditCondition="bEvaluateLensFileForDistortion"))
+	FLensFilePicker LensFilePicker;
+
+	/** Whether to scale the computed overscan by the overscan percentage */
+	UPROPERTY(AdvancedDisplay, BlueprintReadWrite, Category = "Default")
+	bool bScaleOverscan = false;
+
+	/** The percentage of the computed overscan that should be applied to the target camera */
+	UPROPERTY(EditAnywhere, AdvancedDisplay, BlueprintReadWrite, Category = "Default", meta = (EditCondition = "bScaleOverscan", ClampMin = "0.0", ClampMax = "2.0"))
+	float OverscanMultiplier = 1.0f;
+
+	/** Distortion handler produced by this component */
+	UPROPERTY(Transient)
+	ULensDistortionModelHandlerBase* ProducedLensDistortionHandler = nullptr;
+
+	/** Unique identifier representing the source of distortion data */
+	UPROPERTY(DuplicateTransient)
+	FGuid DistortionProducerID;
+
+private:
+	/** Cached rotation, used to track changes to the camera transform, which impacts the original rotation used by nodal offset */
+	FRotator LastRotation;
+
+	/** Cached location, used to track changes to the camera transform, which impacts the original location used by nodal offset */
+	FVector LastLocation;
+};
\ No newline at end of file
