Centralize application of nodal offset to LensComponent, removing all references to nodal offset from the LiveLink camera controller

diff --git a/Engine/Plugins/Animation/LiveLink/Source/LiveLinkComponents/Private/LiveLinkComponentController.cpp b/Engine/Plugins/Animation/LiveLink/Source/LiveLinkComponents/Private/LiveLinkComponentController.cpp
--- a/Engine/Plugins/Animation/LiveLink/Source/LiveLinkComponents/Private/LiveLinkComponentController.cpp
+++ b/Engine/Plugins/Animation/LiveLink/Source/LiveLinkComponents/Private/LiveLinkComponentController.cpp
@@ -2,6 +2,7 @@
 
 #include "LiveLinkComponentController.h"
 
+#include "ILiveLinkComponentModule.h"
 #include "LiveLinkComponentPrivate.h"
 #include "LiveLinkComponentSettings.h"
 #include "LiveLinkControllerBase.h"
@@ -136,6 +137,13 @@ void ULiveLinkComponentController::OnRegister()
 	Super::OnRegister();
 
 	bIsDirty = true;
+
+	ILiveLinkComponentsModule& LiveLinkComponentsModule = FModuleManager::GetModuleChecked<ILiveLinkComponentsModule>(TEXT("LiveLinkComponents"));
+
+	if (LiveLinkComponentsModule.OnLiveLinkComponentRegistered().IsBound())
+	{
+		LiveLinkComponentsModule.OnLiveLinkComponentRegistered().Broadcast(this);
+	}
 }
 
 #if WITH_EDITOR
@@ -214,7 +222,12 @@ void ULiveLinkComponentController::TickComponent(float DeltaTime, ELevelTick Tic
 		FEditorScriptExecutionGuard ScriptGuard;
 		OnLiveLinkUpdated.Broadcast(DeltaTime);
 	}
-	
+
+	if (bHasValidData && OnLiveLinkControllersTicked().IsBound())
+	{
+		OnLiveLinkControllersTicked().Broadcast(this, SubjectData);
+	}
+
 	bIsDirty = false;
 
 	Super::TickComponent(DeltaTime, TickType, ThisTickFunction);
diff --git a/Engine/Plugins/Animation/LiveLink/Source/LiveLinkComponents/Private/LiveLinkComponentModule.cpp b/Engine/Plugins/Animation/LiveLink/Source/LiveLinkComponents/Private/LiveLinkComponentModule.cpp
--- a/Engine/Plugins/Animation/LiveLink/Source/LiveLinkComponents/Private/LiveLinkComponentModule.cpp
+++ b/Engine/Plugins/Animation/LiveLink/Source/LiveLinkComponents/Private/LiveLinkComponentModule.cpp
@@ -1,13 +1,10 @@
 // Copyright Epic Games, Inc. All Rights Reserved.
 
+#include "LiveLinkComponentModule.h"
+
 #include "LiveLinkComponentPrivate.h"
 #include "Modules/ModuleManager.h"
 
-#define LOCTEXT_NAMESPACE "LiveLinkComponentsModule"
-
 DEFINE_LOG_CATEGORY(LogLiveLinkComponents);
 
-
-IMPLEMENT_MODULE(FDefaultModuleImpl, LiveLinkComponents);
-
-#undef LOCTEXT_NAMESPACE
\ No newline at end of file
+IMPLEMENT_MODULE(FLiveLinkComponentsModule, LiveLinkComponents);
diff --git a/Engine/Plugins/Animation/LiveLink/Source/LiveLinkComponents/Private/LiveLinkComponentModule.h b/Engine/Plugins/Animation/LiveLink/Source/LiveLinkComponents/Private/LiveLinkComponentModule.h
new file mode 100644
--- /dev/null
+++ b/Engine/Plugins/Animation/LiveLink/Source/LiveLinkComponents/Private/LiveLinkComponentModule.h
@@ -0,0 +1,14 @@
+// Copyright Epic Games, Inc. All Rights Reserved.
+
+#pragma once
+
+#include "ILiveLinkComponentModule.h"
+
+class FLiveLinkComponentsModule : public ILiveLinkComponentsModule
+{
+public:
+	virtual FLiveLinkComponentRegistered& OnLiveLinkComponentRegistered() override { return OnLiveLinkComponentRegisteredDelegate; }
+
+private:
+	FLiveLinkComponentRegistered OnLiveLinkComponentRegisteredDelegate;
+};
diff --git a/Engine/Plugins/Animation/LiveLink/Source/LiveLinkComponents/Public/ILiveLinkComponentModule.h b/Engine/Plugins/Animation/LiveLink/Source/LiveLinkComponents/Public/ILiveLinkComponentModule.h
new file mode 100644
--- /dev/null
+++ b/Engine/Plugins/Animation/LiveLink/Source/LiveLinkComponents/Public/ILiveLinkComponentModule.h
@@ -0,0 +1,24 @@
+// Copyright Epic Games, Inc. All Rights Reserved.
+
+#pragma once
+
+#include "Modules/ModuleInterface.h"
+#include "Modules/ModuleManager.h"
+
+#include "LiveLinkComponentController.h"
+
+DECLARE_MULTICAST_DELEGATE_OneParam(FLiveLinkComponentRegistered, ULiveLinkComponentController*);
+
+/**
+ * Interface for LiveLinkComponent module
+ */
+class ILiveLinkComponentsModule : public IModuleInterface
+{
+public:
+	virtual FLiveLinkComponentRegistered& OnLiveLinkComponentRegistered() = 0;
+
+public:
+	/** Virtual destructor. */
+	virtual ~ILiveLinkComponentsModule() { }
+};
+
diff --git a/Engine/Plugins/Animation/LiveLink/Source/LiveLinkComponents/Public/LiveLinkComponentController.h b/Engine/Plugins/Animation/LiveLink/Source/LiveLinkComponents/Public/LiveLinkComponentController.h
--- a/Engine/Plugins/Animation/LiveLink/Source/LiveLinkComponents/Public/LiveLinkComponentController.h
+++ b/Engine/Plugins/Animation/LiveLink/Source/LiveLinkComponents/Public/LiveLinkComponentController.h
@@ -9,6 +9,8 @@
 DECLARE_DYNAMIC_MULTICAST_DELEGATE_OneParam(FLiveLinkTickDelegate, float, DeltaTime);
 DECLARE_DYNAMIC_MULTICAST_DELEGATE(FOnControllerMapUpdatedDelegate);
 
+DECLARE_MULTICAST_DELEGATE_TwoParams(FLiveLinkControllersTicked, const ULiveLinkComponentController* const, const FLiveLinkSubjectFrameData&);
+
 class ULiveLinkControllerBase;
 
 UCLASS( ClassGroup=(LiveLink), meta=(DisplayName="LiveLink Controller", BlueprintSpawnableComponent) )
@@ -94,6 +96,9 @@ public:
 	/** Set the component to control for the LiveLink controller of the input Role */
 	void SetControlledComponent(TSubclassOf<ULiveLinkRole> InRoleClass, UActorComponent* InComponent);
 
+	/** Multicast delegate that broadcasts after LiveLink controllers have ticked with the latest frame of subject data */
+	FLiveLinkControllersTicked& OnLiveLinkControllersTicked() { return LiveLinkControllersTickedDelegate; }
+
 #if WITH_EDITOR
 	/** Used to cleanup controllers when exiting PIE */
 	void OnEndPIE(bool bIsSimulating);
@@ -129,4 +134,7 @@ protected:
 	/** Called during loading to convert old data to new scheme. */
 	void ConvertOldControllerSystem();
 #endif //WITH_EDITOR
-};
\ No newline at end of file
+
+private:
+	FLiveLinkControllersTicked LiveLinkControllersTickedDelegate;
+};
diff --git a/Engine/Plugins/VirtualProduction/CameraCalibration/Source/CameraCalibrationEditor/Private/AssetEditor/NodalOffsetTool.cpp b/Engine/Plugins/VirtualProduction/CameraCalibration/Source/CameraCalibrationEditor/Private/AssetEditor/NodalOffsetTool.cpp
--- a/Engine/Plugins/VirtualProduction/CameraCalibration/Source/CameraCalibrationEditor/Private/AssetEditor/NodalOffsetTool.cpp
+++ b/Engine/Plugins/VirtualProduction/CameraCalibration/Source/CameraCalibrationEditor/Private/AssetEditor/NodalOffsetTool.cpp
@@ -3,11 +3,12 @@
 #include "NodalOffsetTool.h"
 
 #include "CalibrationPointComponent.h"
+#include "Camera/CameraActor.h"
 #include "CameraCalibrationStepsController.h"
 #include "CameraCalibrationSubsystem.h"
 #include "CameraCalibrationTypes.h"
 #include "ImageCenterTool.h"
-#include "LiveLinkCameraController.h"
+#include "LensComponent.h"
 #include "Misc/MessageDialog.h"
 #include "ScopedTransaction.h"
 #include "SNodalOffsetToolPanel.h"
@@ -229,9 +230,16 @@ void UNodalOffsetTool::OnSaveCurrentNodalOffset()
 	LensFile->AddNodalOffsetPoint(Focus, Zoom, NodalOffset);
 
 	// Force bApplyNodalOffset in the LiveLinkCameraController so that we can see the effect right away
-	if (ULiveLinkCameraController* LiveLinkCameraController = CameraCalibrationStepsController.Pin()->FindLiveLinkCameraController())
+	TInlineComponentArray<ULensComponent*> LensComponents;
+	CameraCalibrationStepsController.Pin()->GetCamera()->GetComponents(LensComponents);
+
+	for (ULensComponent* LensComponent : LensComponents)
 	{
-		LiveLinkCameraController->SetApplyNodalOffset(true);
+		if (LensFile == LensComponent->GetLensFile())
+		{
+			LensComponent->SetApplyNodalOffsetOnTick(true);
+			break;
+		}
 	}
 
 	Algo->OnSavedNodalOffset();
diff --git a/Engine/Plugins/VirtualProduction/CameraCalibrationCore/Config/BaseCameraCalibrationCore.ini b/Engine/Plugins/VirtualProduction/CameraCalibrationCore/Config/BaseCameraCalibrationCore.ini
--- a/Engine/Plugins/VirtualProduction/CameraCalibrationCore/Config/BaseCameraCalibrationCore.ini
+++ b/Engine/Plugins/VirtualProduction/CameraCalibrationCore/Config/BaseCameraCalibrationCore.ini
@@ -37,3 +37,7 @@
 ; Packages
 
 +PackageRedirects=(OldName="/Script/CameraCalibration", Removed=true)
+
+; Properties
+
++PropertyRedirects=(OldName="LensComponent.bApplyNodalOffset",NewName="LensComponent.bApplyNodalOffsetOnTick")
diff --git a/Engine/Plugins/VirtualProduction/CameraCalibrationCore/Source/CameraCalibrationCore/CameraCalibrationCore.Build.cs b/Engine/Plugins/VirtualProduction/CameraCalibrationCore/Source/CameraCalibrationCore/CameraCalibrationCore.Build.cs
--- a/Engine/Plugins/VirtualProduction/CameraCalibrationCore/Source/CameraCalibrationCore/CameraCalibrationCore.Build.cs
+++ b/Engine/Plugins/VirtualProduction/CameraCalibrationCore/Source/CameraCalibrationCore/CameraCalibrationCore.Build.cs
@@ -26,6 +26,7 @@ namespace UnrealBuildTool.Rules
 				new string[]
                 {
 					"LiveLinkComponents",
+					"LiveLinkInterface",
                     "Projects",
 				}
 			);
diff --git a/Engine/Plugins/VirtualProduction/CameraCalibrationCore/Source/CameraCalibrationCore/Private/LensComponent.cpp b/Engine/Plugins/VirtualProduction/CameraCalibrationCore/Source/CameraCalibrationCore/Private/LensComponent.cpp
--- a/Engine/Plugins/VirtualProduction/CameraCalibrationCore/Source/CameraCalibrationCore/Private/LensComponent.cpp
+++ b/Engine/Plugins/VirtualProduction/CameraCalibrationCore/Source/CameraCalibrationCore/Private/LensComponent.cpp
@@ -4,21 +4,27 @@
 
 #include "CameraCalibrationSubsystem.h"
 #include "CineCameraComponent.h"
+#include "Controllers/LiveLinkTransformController.h"
 #include "Engine/Engine.h"
-#include "LiveLinkComponentController.h"
+#include "ILiveLinkComponentModule.h"
+#include "Roles/LiveLinkCameraTypes.h"
+#include "Roles/LiveLinkCameraRole.h"
+#include "Roles/LiveLinkTransformRole.h"
+#include "UObject/UE5MainStreamObjectVersion.h"
+
+DEFINE_LOG_CATEGORY_STATIC(LogLensComponent, Log, All);
 
 ULensComponent::ULensComponent()
 {
 	PrimaryComponentTick.bCanEverTick = true;
 	PrimaryComponentTick.bStartWithTickEnabled = true;
-	PrimaryComponentTick.TickGroup = ETickingGroup::TG_LastDemotable;
+	// This component is designed to tick after the LiveLink component, which uses TG_PrePhysics
+	// We also use a tick prerequisite on LiveLink components, so technically this could also use TG_PrePhysics
+	PrimaryComponentTick.TickGroup = ETickingGroup::TG_DuringPhysics;
 	bTickInEditor = true;
 
 	if (!HasAnyFlags(RF_ArchetypeObject | RF_ClassDefaultObject))
 	{
-		//Hook up to PostActorTick to handle nodal offset
-		FWorldDelegates::OnWorldPostActorTick.AddUObject(this, &ULensComponent::OnPostActorTick);
-
 		DistortionProducerID = FGuid::NewGuid();
 	}
 }
@@ -29,46 +35,58 @@ void ULensComponent::OnRegister()
 
 	InitDefaultCamera();
 
-	if (UCineCameraComponent* const CineCameraComponent = Cast<UCineCameraComponent>(TargetCameraComponent.GetComponent(GetOwner())))
+	// Register a callback to let us know when a new LiveLinkComponent is added to this component's parent actor
+	// This gives us the opportunity to track when the subject representation of that LiveLinkComponent changes
+	ILiveLinkComponentsModule& LiveLinkComponentsModule = FModuleManager::GetModuleChecked<ILiveLinkComponentsModule>(TEXT("LiveLinkComponents"));
+	if (!LiveLinkComponentsModule.OnLiveLinkComponentRegistered().IsBoundToObject(this))
 	{
-		LastRotation = CineCameraComponent->GetRelativeRotation();
-		LastLocation = CineCameraComponent->GetRelativeLocation();
+		LiveLinkComponentsModule.OnLiveLinkComponentRegistered().AddUObject(this, &ULensComponent::OnLiveLinkComponentRegistered);
 	}
-}
 
-void ULensComponent::DestroyComponent(bool bPromoteChildren /*= false*/)
-{
-	Super::DestroyComponent(bPromoteChildren);
+	// Look for any LiveLinkComponents that were previously added to this component's parent actor
+	TInlineComponentArray<ULiveLinkComponentController*> LiveLinkComponents;
+	GetOwner()->GetComponents(LiveLinkComponents);
 
-	if (UCineCameraComponent* const CineCameraComponent = Cast<UCineCameraComponent>(TargetCameraComponent.GetComponent(GetOwner())))
-	{
-		CleanupDistortion(CineCameraComponent);
+	for (ULiveLinkComponentController* LiveLinkComponent : LiveLinkComponents)
+	{	
+		AddTickPrerequisiteComponent(LiveLinkComponent);
 
-		if (UCameraCalibrationSubsystem* SubSystem = GEngine->GetEngineSubsystem<UCameraCalibrationSubsystem>())
+		if (!LiveLinkComponent->OnLiveLinkControllersTicked().IsBoundToObject(this))
 		{
-			SubSystem->UnregisterDistortionModelHandler(CineCameraComponent, ProducedLensDistortionHandler);
+			LiveLinkComponent->OnLiveLinkControllersTicked().AddUObject(this, &ULensComponent::ProcessLiveLinkData);
 		}
+	}
+}
+
+void ULensComponent::OnUnregister()
+{
+	Super::OnUnregister();
 
-		// Restore the target camera's location and rotation (undo nodal offset)
-		CineCameraComponent->SetRelativeLocation(OriginalCameraLocation);
-		CineCameraComponent->SetRelativeRotation(OriginalCameraRotation.Quaternion());
+	ILiveLinkComponentsModule& LiveLinkComponentsModule = FModuleManager::GetModuleChecked<ILiveLinkComponentsModule>(TEXT("LiveLinkComponents"));
+	LiveLinkComponentsModule.OnLiveLinkComponentRegistered().RemoveAll(this);
+
+	TInlineComponentArray<ULiveLinkComponentController*> LiveLinkComponents;
+	GetOwner()->GetComponents(LiveLinkComponents);
+
+	for (ULiveLinkComponentController* LiveLinkComponent : LiveLinkComponents)
+	{
+		LiveLinkComponent->OnLiveLinkControllersTicked().RemoveAll(this);
 	}
 }
 
 void ULensComponent::TickComponent(float DeltaTime, ELevelTick TickType, FActorComponentTickFunction* ThisTickFunction)
 {
-	if (UCineCameraComponent* const CineCameraComponent = Cast<UCineCameraComponent>(TargetCameraComponent.GetComponent(GetOwner())))
-	{
-		// Add a tick prerequisite on any live link components also on this camera
-		// LiveLink might set the focal length of the target camera, so it is important that this component ticks after that
-		TInlineComponentArray<ULiveLinkComponentController*> LiveLinkComponents;
-		CineCameraComponent->GetOwner()->GetComponents(LiveLinkComponents);
+	// Get the focus and zoom values needed to evaluate the LensFile this Tick
+	UpdateLensFileEvaluationInputs();
 
-		for (UActorComponent* LiveLinkComponent : LiveLinkComponents)
-		{
-			AddTickPrerequisiteComponent(LiveLinkComponent);
-		}
+	// Attempt to apply nodal offset, which will only succeed if there is a valid LensFile, evaluation inputs, and component to offset
+	if (bApplyNodalOffsetOnTick)
+	{
+		ApplyNodalOffset();
+	}
 
+	if (UCineCameraComponent* const CineCameraComponent = Cast<UCineCameraComponent>(TargetCameraComponent.GetComponent(GetOwner())))
+	{
 		DistortionSource.TargetCameraComponent = CineCameraComponent;
 
 		UCameraCalibrationSubsystem* const SubSystem = GEngine->GetEngineSubsystem<UCameraCalibrationSubsystem>();
@@ -106,7 +124,7 @@ void ULensComponent::TickComponent(float DeltaTime, ELevelTick TickType, FActorC
 				SubSystem->UpdateOriginalFocalLength(CineCameraComponent, CineCameraComponent->CurrentFocalLength);
 			}
 		}
-	
+
 		if (bApplyDistortion)
 		{
 			// Get the lens distortion handler for the target camera and distortion source
@@ -157,34 +175,91 @@ void ULensComponent::TickComponent(float DeltaTime, ELevelTick TickType, FActorC
 	Super::TickComponent(DeltaTime, TickType, ThisTickFunction);
 }
 
-#if WITH_EDITOR
-void ULensComponent::PostEditChangeProperty(struct FPropertyChangedEvent& PropertyChangedEvent)
+void ULensComponent::DestroyComponent(bool bPromoteChildren /*= false*/)
 {
- 	const FName PropertyName = (PropertyChangedEvent.Property != NULL) ? PropertyChangedEvent.Property->GetFName() : NAME_None;
- 
- 	if (PropertyName == GET_MEMBER_NAME_CHECKED(ULensComponent, bApplyDistortion))
- 	{
- 		if (UCineCameraComponent* const CineCameraComponent = Cast<UCineCameraComponent>(TargetCameraComponent.GetComponent(GetOwner())))
- 		{
- 			if (!bApplyDistortion)
- 			{
-				CleanupDistortion(CineCameraComponent);
- 			}
- 		}
- 	}
-	else if (PropertyName == GET_MEMBER_NAME_CHECKED(ULensComponent, bApplyNodalOffset))
+	Super::DestroyComponent(bPromoteChildren);
+
+	if (UCineCameraComponent* const CineCameraComponent = Cast<UCineCameraComponent>(TargetCameraComponent.GetComponent(GetOwner())))
+	{
+		CleanupDistortion(CineCameraComponent);
+
+		if (UCameraCalibrationSubsystem* SubSystem = GEngine->GetEngineSubsystem<UCameraCalibrationSubsystem>())
+		{
+			SubSystem->UnregisterDistortionModelHandler(CineCameraComponent, ProducedLensDistortionHandler);
+		}
+	}
+}
+
+void ULensComponent::PostDuplicate(bool bDuplicateForPIE)
+{
+	Super::PostDuplicate(bDuplicateForPIE);
+
+	InitDefaultCamera();
+
+	// When this component is duplicated (e.g. for PIE), the duplicated component needs a new unique ID
+	if (!DistortionProducerID.IsValid())
+	{
+		DistortionProducerID = FGuid::NewGuid();
+	}
+}
+
+void ULensComponent::PostEditImport()
+{
+	Super::PostEditImport();
+
+	InitDefaultCamera();
+
+	// When this component is duplicated (e.g. copy-paste), the duplicated component needs a new unique ID
+	if (!DistortionProducerID.IsValid())
+	{
+		DistortionProducerID = FGuid::NewGuid();
+	}
+}
+
+void ULensComponent::PostLoad()
+{
+	Super::PostLoad();
+
+	if (bIsDistortionSetup)
 	{
 		if (UCineCameraComponent* const CineCameraComponent = Cast<UCineCameraComponent>(TargetCameraComponent.GetComponent(GetOwner())))
 		{
-			if (bApplyNodalOffset)
+			CleanupDistortion(CineCameraComponent);
+		}
+		bIsDistortionSetup = false;
+	}
+
+#if WITH_EDITOR
+	const int32 UE5MainVersion = GetLinkerCustomVersion(FUE5MainStreamObjectVersion::GUID);
+	if (UE5MainVersion < FUE5MainStreamObjectVersion::LensComponentNodalOffset)
+	{
+		// If this component was previously applying nodal offset, reset the camera component's transform to its original pose
+		if (UCineCameraComponent* const CineCameraComponent = Cast<UCineCameraComponent>(TargetCameraComponent.GetComponent(GetOwner())))
+		{
+			PRAGMA_DISABLE_DEPRECATION_WARNINGS
+			if (bApplyNodalOffsetOnTick)
 			{
-				OriginalCameraRotation = CineCameraComponent->GetRelativeRotation();
-				OriginalCameraLocation = CineCameraComponent->GetRelativeLocation();
+				CineCameraComponent->SetRelativeLocation(OriginalCameraLocation_DEPRECATED);
+				CineCameraComponent->SetRelativeRotation(OriginalCameraRotation_DEPRECATED.Quaternion());
 			}
-			else
+			PRAGMA_ENABLE_DEPRECATION_WARNINGS
+		}
+	}
+#endif //WITH_EDITOR
+}
+
+#if WITH_EDITOR
+void ULensComponent::PostEditChangeProperty(struct FPropertyChangedEvent& PropertyChangedEvent)
+{
+	const FName PropertyName = (PropertyChangedEvent.Property != NULL) ? PropertyChangedEvent.Property->GetFName() : NAME_None;
+
+	if (PropertyName == GET_MEMBER_NAME_CHECKED(ULensComponent, bApplyDistortion))
+	{
+		if (UCineCameraComponent* const CineCameraComponent = Cast<UCineCameraComponent>(TargetCameraComponent.GetComponent(GetOwner())))
+		{
+			if (!bApplyDistortion)
 			{
-				CineCameraComponent->SetRelativeLocation(OriginalCameraLocation);
-				CineCameraComponent->SetRelativeRotation(OriginalCameraRotation.Quaternion());
+				CleanupDistortion(CineCameraComponent);
 			}
 		}
 	}
@@ -217,13 +292,6 @@ void ULensComponent::PostEditChangeProperty(struct FPropertyChangedEvent& Proper
 			}
 
 			ProducedLensDistortionHandler = nullptr;
-
-			// Restore the target camera's location and rotation (undo nodal offset)
-			if (UCineCameraComponent* const CineCameraComponent = Cast<UCineCameraComponent>(TargetCameraComponent.GetComponent(GetOwner())))
-			{
-				CineCameraComponent->SetRelativeLocation(OriginalCameraLocation);
-				CineCameraComponent->SetRelativeRotation(OriginalCameraRotation.Quaternion());
-			}
 		}
 	}
 	else if (PropertyName == GET_MEMBER_NAME_CHECKED(ULensComponent, TargetCameraComponent))
@@ -243,10 +311,6 @@ void ULensComponent::PostEditChangeProperty(struct FPropertyChangedEvent& Proper
 
 				ProducedLensDistortionHandler = nullptr;
 			}
-
-			// Restore the last camera's location and rotation (undo nodal offset)
-			LastCameraComponent->SetRelativeLocation(OriginalCameraLocation);
-			LastCameraComponent->SetRelativeRotation(OriginalCameraRotation.Quaternion());
 		}
 
 		LastCameraComponent = Cast<UCineCameraComponent>(TargetCameraComponent.GetComponent(GetOwner()));
@@ -257,13 +321,6 @@ void ULensComponent::PostEditChangeProperty(struct FPropertyChangedEvent& Proper
 		{
 			CreateDistortionHandlerForLensFile();
 		}
-
-		// Cache the rotation and location of the new target camera component
-		if (UCineCameraComponent* const CineCameraComponent = Cast<UCineCameraComponent>(TargetCameraComponent.GetComponent(GetOwner())))
-		{
-			LastRotation = OriginalCameraRotation = CineCameraComponent->GetRelativeRotation();
-			LastLocation = OriginalCameraLocation = CineCameraComponent->GetRelativeLocation();
-		}
 	}
 	else if (PropertyName == GET_MEMBER_NAME_CHECKED(FLensFilePicker, LensFile))
 	{
@@ -277,56 +334,117 @@ void ULensComponent::PostEditChangeProperty(struct FPropertyChangedEvent& Proper
 }
 #endif //WITH_EDITOR
 
-void ULensComponent::PostDuplicate(bool bDuplicateForPIE)
+void ULensComponent::OnPostActorTick(UWorld* World, ELevelTick TickType, float DeltaSeconds)
 {
-	Super::PostDuplicate(bDuplicateForPIE);
+	// The use of this callback by this class has been deprecated
+}
 
-	InitDefaultCamera();
+void ULensComponent::ApplyNodalOffset()
+{
+	// Verify that we detected a tracked component for us to offset this tick. If there is none, it is likely because no LiveLink transform controller executed this tick.
+	if (!TrackedComponent.IsValid())
+	{
+		return;
+	}
 
-	// When this component is duplicated (e.g. for PIE), the duplicated component needs a new unique ID
-	if (!DistortionProducerID.IsValid())
+	// Verify that the evaluation inputs are valid
+	if (EvaluationMode == EFIZEvaluationMode::DoNotEvaluate)
 	{
-		DistortionProducerID = FGuid::NewGuid();
+		return;
+	}
+
+	// Verify that there is a valid LensFile to evaluate
+	ULensFile* LensFile = LensFilePicker.GetLensFile();
+	if (!LensFile)
+	{
+		return;
 	}
+
+	FNodalPointOffset Offset;
+	LensFile->EvaluateNodalPointOffset(EvalFocus, EvalZoom, Offset);
+
+	// TODO: This would be a good time to cache the raw tracking pose before applying the nodal offset
+
+	TrackedComponent.Get()->AddLocalOffset(Offset.LocationOffset);
+	TrackedComponent.Get()->AddLocalRotation(Offset.RotationOffset);
+
+	// Reset so that nodal offset will only be applied again next tick if new tracking data was applied between now and then
+	TrackedComponent.Reset();
 }
 
-void ULensComponent::PostEditImport()
+void ULensComponent::ApplyNodalOffset(USceneComponent* ComponentToOffset, bool bUseManualInputs, float ManualFocusInput, float ManualZoomInput)
 {
-	Super::PostEditImport();
+	bApplyNodalOffsetOnTick = false;
 
-	InitDefaultCamera();
+	// Verify that the input component was not null
+	if (!ComponentToOffset)
+	{
+		return;
+	}
 
-	// When this component is duplicated (e.g. copy-paste), the duplicated component needs a new unique ID
-	if (!DistortionProducerID.IsValid())
+	// Verify that there is a valid LensFile to evaluate
+	ULensFile* LensFile = LensFilePicker.GetLensFile();
+	if (!LensFile)
 	{
-		DistortionProducerID = FGuid::NewGuid();
+		return;
+	}
+
+	FNodalPointOffset Offset;
+	if (bUseManualInputs)
+	{
+		LensFile->EvaluateNodalPointOffset(ManualFocusInput, ManualZoomInput, Offset);
 	}
+	else
+	{
+		LensFile->EvaluateNodalPointOffset(EvalFocus, EvalZoom, Offset);
+	}
+
+	ComponentToOffset->AddLocalOffset(Offset.LocationOffset);
+	ComponentToOffset->AddLocalRotation(Offset.RotationOffset);
 }
 
-void ULensComponent::PostLoad()
+ULensFile* ULensComponent::GetLensFile() const
 {
-	Super::PostLoad();
+	return LensFilePicker.GetLensFile();
+}
 
-	if (bIsDistortionSetup)
-	{
-		if (UCineCameraComponent* const CineCameraComponent = Cast<UCineCameraComponent>(TargetCameraComponent.GetComponent(GetOwner())))
-		{
-			CleanupDistortion(CineCameraComponent);
-		}
-		bIsDistortionSetup = false;
-	}
+void ULensComponent::SetLensFile(FLensFilePicker LensFile)
+{
+	LensFilePicker = LensFile;
 }
 
-void ULensComponent::OnPostActorTick(UWorld* World, ELevelTick TickType, float DeltaSeconds)
+EFIZEvaluationMode ULensComponent::GetFIZEvaluationMode() const
+{
+	return EvaluationMode;
+}
+
+void ULensComponent::SetFIZEvaluationMode(EFIZEvaluationMode Mode)
+{
+	EvaluationMode = Mode;
+}
+
+bool ULensComponent::ShouldApplyNodalOffsetOnTick() const
+{
+	return bApplyNodalOffsetOnTick;
+}
+
+void ULensComponent::SetApplyNodalOffsetOnTick(bool bApplyNodalOffset)
 {
-	if (World == GetWorld())
+	bApplyNodalOffsetOnTick = bApplyNodalOffset;
+}
+
+void ULensComponent::InitDefaultCamera()
+{
+	UCineCameraComponent* const CineCameraComponent = Cast<UCineCameraComponent>(TargetCameraComponent.GetComponent(GetOwner()));
+	if (!CineCameraComponent)
 	{
-		if (UCineCameraComponent* const CineCameraComponent = Cast<UCineCameraComponent>(TargetCameraComponent.GetComponent(GetOwner())))
+		// Find all CineCameraComponents on the same actor as this component and set the first one to be the target
+		TInlineComponentArray<UCineCameraComponent*> CineCameraComponents;
+		GetOwner()->GetComponents(CineCameraComponents);
+		if (CineCameraComponents.Num() > 0)
 		{
-			if (ULensFile* CurrentLensFile = LensFilePicker.GetLensFile())
-			{
-				ApplyNodalOffset(CurrentLensFile, CineCameraComponent);
-			}
+			TargetCameraComponent.ComponentProperty = CineCameraComponents[0]->GetFName();
+			LastCameraComponent = CineCameraComponents[0];
 		}
 	}
 }
@@ -360,25 +478,6 @@ void ULensComponent::CleanupDistortion(UCineCameraComponent* const CineCameraCom
 	bIsDistortionSetup = false;
 }
 
-void ULensComponent::InitDefaultCamera()
-{
-	UCineCameraComponent* const CineCameraComponent = Cast<UCineCameraComponent>(TargetCameraComponent.GetComponent(GetOwner()));
-	if (!CineCameraComponent)
-	{
-		// Find all CineCameraComponents on the same actor as this component and set the first one to be the target
-		TInlineComponentArray<UCineCameraComponent*> CineCameraComponents;
-		GetOwner()->GetComponents(CineCameraComponents);
-		if (CineCameraComponents.Num() > 0)
-		{
-			TargetCameraComponent.ComponentProperty = CineCameraComponents[0]->GetFName();
-			LastCameraComponent = CineCameraComponents[0];
-
-			OriginalCameraRotation = LastCameraComponent->GetRelativeRotation();
-			OriginalCameraLocation = LastCameraComponent->GetRelativeLocation();
-		}
-	}
-}
-
 void ULensComponent::CreateDistortionHandlerForLensFile()
 {
 	if (DistortionSource.TargetCameraComponent == nullptr)
@@ -402,66 +501,105 @@ void ULensComponent::CreateDistortionHandlerForLensFile()
 	}
 }
 
-void ULensComponent::ApplyNodalOffset(ULensFile* SelectedLensFile, UCineCameraComponent* CineCameraComponent)
+void ULensComponent::OnLiveLinkComponentRegistered(ULiveLinkComponentController* LiveLinkComponent)
 {
-	if (CineCameraComponent)
+	// Check that the new LiveLinkComponent that was just registered was added to the same actor that this component belongs to
+	if (LiveLinkComponent->GetOwner() == GetOwner())
 	{
-		const FRotator CurrentRotator = CineCameraComponent->GetRelativeRotation();
-		const FVector CurrentTranslation = CineCameraComponent->GetRelativeLocation();
+		AddTickPrerequisiteComponent(LiveLinkComponent);
 
-		//Verify if something was set by user / programmatically and update original values
-		if (CurrentRotator != LastRotation)
+		if (!LiveLinkComponent->OnLiveLinkControllersTicked().IsBoundToObject(this))
 		{
-			if (CurrentRotator.Pitch != LastRotation.Pitch)
-			{
-				OriginalCameraRotation.Pitch = CurrentRotator.Pitch;
-			}
+			LiveLinkComponent->OnLiveLinkControllersTicked().AddUObject(this, &ULensComponent::ProcessLiveLinkData);
+		}
+	}
+}
 
-			if (CurrentRotator.Yaw != LastRotation.Yaw)
-			{
-				OriginalCameraRotation.Yaw = CurrentRotator.Yaw;
-			}
+void ULensComponent::ProcessLiveLinkData(const ULiveLinkComponentController* const LiveLinkComponent, const FLiveLinkSubjectFrameData& SubjectData)
+{
+	TSubclassOf<ULiveLinkRole> LiveLinkRole = LiveLinkComponent->GetSubjectRepresentation().Role;
+	if (LiveLinkRole == ULiveLinkCameraRole::StaticClass())
+	{
+		UpdateTrackedComponent(LiveLinkComponent, SubjectData);
+		UpdateLiveLinkFIZ(LiveLinkComponent, SubjectData);
+	}
+	else if (LiveLinkRole == ULiveLinkTransformRole::StaticClass())
+	{
+		UpdateTrackedComponent(LiveLinkComponent, SubjectData);
+	}
+}
 
-			if (CurrentRotator.Roll != LastRotation.Roll)
+void ULensComponent::UpdateTrackedComponent(const ULiveLinkComponentController* const LiveLinkComponent, const FLiveLinkSubjectFrameData& SubjectData)
+{
+	const FLiveLinkTransformStaticData* StaticData = SubjectData.StaticData.Cast<FLiveLinkTransformStaticData>();
+	const FLiveLinkTransformFrameData* FrameData = SubjectData.FrameData.Cast<FLiveLinkTransformFrameData>();
+
+	check(StaticData && FrameData);
+
+	// Find the transform controller in the LiveLink component's controller map
+	if (const TObjectPtr<ULiveLinkControllerBase>* TransformControllerPtr = LiveLinkComponent->ControllerMap.Find(ULiveLinkTransformRole::StaticClass()))
+	{
+		if (ULiveLinkTransformController* TransformController = Cast<ULiveLinkTransformController>(*TransformControllerPtr))
+		{
+			// Check LiveLink static data to ensure that location and rotation were supported by this subject
+			bool bIsLocationRotationSupported = (StaticData->bIsLocationSupported && StaticData->bIsRotationSupported);
+
+			// Check the transform controller usage flags to ensure that location and rotation were supported by the controller
+			bIsLocationRotationSupported = bIsLocationRotationSupported && (TransformController->TransformData.bUseLocation && TransformController->TransformData.bUseRotation);
+
+			if (bIsLocationRotationSupported)
 			{
-				OriginalCameraRotation.Roll = CurrentRotator.Roll;
+				TrackedComponent = Cast<USceneComponent>(TransformController->GetAttachedComponent());
 			}
 		}
+	}
+}
 
-		if (CurrentTranslation != LastLocation)
+void ULensComponent::UpdateLiveLinkFIZ(const ULiveLinkComponentController* const LiveLinkComponent, const FLiveLinkSubjectFrameData& SubjectData)
+{
+	const FLiveLinkCameraStaticData* StaticData = SubjectData.StaticData.Cast<FLiveLinkCameraStaticData>();
+	const FLiveLinkCameraFrameData* FrameData = SubjectData.FrameData.Cast<FLiveLinkCameraFrameData>();
+	
+	check(StaticData && FrameData);
+
+	// Check that the camera component being controlled by the input LiveLink component matches our target camera
+	if (const TObjectPtr<ULiveLinkControllerBase>* CameraControllerPtr = LiveLinkComponent->ControllerMap.Find(ULiveLinkCameraRole::StaticClass()))
+	{
+		if (ULiveLinkControllerBase* CameraController = *CameraControllerPtr)
 		{
-			if (CurrentTranslation.X != LastLocation.X)
-			{
-				OriginalCameraLocation.X = CurrentTranslation.X;
-			}
-			if (CurrentTranslation.Y != LastLocation.Y)
-			{
-				OriginalCameraLocation.Y = CurrentTranslation.Y;
-			}
-			if (CurrentTranslation.Z != LastLocation.Z)
+			if (CameraController->GetAttachedComponent() != TargetCameraComponent.GetComponent(GetOwner()))
 			{
-				OriginalCameraLocation.Z = CurrentTranslation.Z;
+				return;
 			}
 		}
+	}
 
-		if (bEvaluateLensFileForDistortion && bApplyNodalOffset && SelectedLensFile)
-		{
-			FNodalPointOffset Offset;
-
-			SelectedLensFile->EvaluateNodalPointOffset(CineCameraComponent->CurrentFocusDistance, CineCameraComponent->CurrentFocalLength, Offset);
+	if (StaticData->bIsFocusDistanceSupported && StaticData->bIsFocalLengthSupported)
+	{
+		LiveLinkFIZ.Input.Focus = FrameData->FocusDistance;
+		LiveLinkFIZ.Input.Zoom = FrameData->FocalLength;
+	}
+}
 
-			CineCameraComponent->SetRelativeLocation(OriginalCameraLocation);
-			CineCameraComponent->SetRelativeRotation(OriginalCameraRotation.Quaternion());
-			CineCameraComponent->AddLocalOffset(Offset.LocationOffset);
-			CineCameraComponent->AddLocalRotation(Offset.RotationOffset);
+void ULensComponent::UpdateLensFileEvaluationInputs()
+{
+	// Query for the focus and zoom inputs to use to evaluate the LensFile. The source of these inputs will depend on the evaluation mode.
+	if (EvaluationMode == EFIZEvaluationMode::UseLiveLink)
+	{
+		EvalFocus = LiveLinkFIZ.Input.Focus;
+		EvalZoom = LiveLinkFIZ.Input.Zoom;
+	}
+	else if (EvaluationMode == EFIZEvaluationMode::UseCameraSettings)
+	{
+		if (const UCineCameraComponent* const CineCameraComponent = Cast<UCineCameraComponent>(TargetCameraComponent.GetComponent(GetOwner())))
+		{
+			EvalFocus = CineCameraComponent->CurrentFocusDistance;
+			EvalZoom = CineCameraComponent->CurrentFocalLength;
 		}
-
-		LastLocation = CineCameraComponent->GetRelativeLocation();
-		LastRotation = CineCameraComponent->GetRelativeRotation();
 	}
-	else
+	else if (EvaluationMode == EFIZEvaluationMode::UseRecordedValues)
 	{
-		LastLocation = FVector::OneVector;
-		LastRotation = FRotator::ZeroRotator;
+		// Do nothing, the values for EvalFocus and EvalZoom are already loaded from the recorded sequence
 	}
 }
+
diff --git a/Engine/Plugins/VirtualProduction/CameraCalibrationCore/Source/CameraCalibrationCore/Public/LensComponent.h b/Engine/Plugins/VirtualProduction/CameraCalibrationCore/Source/CameraCalibrationCore/Public/LensComponent.h
--- a/Engine/Plugins/VirtualProduction/CameraCalibrationCore/Source/CameraCalibrationCore/Public/LensComponent.h
+++ b/Engine/Plugins/VirtualProduction/CameraCalibrationCore/Source/CameraCalibrationCore/Public/LensComponent.h
@@ -6,13 +6,28 @@
 
 #include "CameraCalibrationTypes.h"
 #include "LensFile.h"
+#include "LiveLinkComponentController.h"
 
 #include "LensComponent.generated.h"
 
 class UCineCameraComponent;
 
+/** Mode that controls where FIZ inputs are sourced from and how they are used to evaluate the LensFile */
+UENUM(BlueprintType)
+enum class EFIZEvaluationMode : uint8
+{
+	/** Evaluate the Lens File with the latest FIZ data received from LiveLink */
+	UseLiveLink,
+	/** Evaluate the Lens File using the current FIZ settings of the target camera */
+	UseCameraSettings,
+	/** Evaluate the Lens File using values recorded in a level sequence (set automatically when the sequence is opened) */
+	UseRecordedValues,
+	/** Do not evaluate the Lens File */
+	DoNotEvaluate,
+};
+
 /** Component for applying a post-process lens distortion effect to a CineCameraComponent on the same actor */
-UCLASS(HideCategories = (Tags, Activation, Cooking, AssetUserData, Collision), meta = (BlueprintSpawnableComponent))
+UCLASS(HideCategories=(Tags, Activation, Cooking, AssetUserData, Collision), meta=(BlueprintSpawnableComponent))
 class CAMERACALIBRATIONCORE_API ULensComponent : public UActorComponent
 {
 	GENERATED_BODY()
@@ -22,24 +37,61 @@ public:
 
 	//~ Begin UActorComponent interface
 	virtual void OnRegister() override;
+	virtual void OnUnregister() override;
 	virtual void TickComponent(float DeltaTime, ELevelTick TickType, FActorComponentTickFunction* ThisTickFunction) override;
 	virtual void DestroyComponent(bool bPromoteChildren = false) override;
 	//~ End UActorComponent interface
 
-	//~ UObject interface
-#if WITH_EDITOR
-	virtual void PostEditChangeProperty(struct FPropertyChangedEvent& PropertyChangedEvent) override;
-#endif //WITH_EDITOR
+	//~ Begin UObject interface
 	virtual void PostDuplicate(bool bDuplicateForPIE) override;
 	virtual void PostEditImport() override;
 	virtual void PostLoad() override;
+
+#if WITH_EDITOR
+	virtual void PostEditChangeProperty(struct FPropertyChangedEvent& PropertyChangedEvent) override;
+#endif //WITH_EDITOR
 	//~ End UObject interface
 
+	/** Get the LensFile used by this component */
+	UFUNCTION(BlueprintPure, Category="Lens Component")
+	ULensFile* GetLensFile() const;
+
+	/** Set the LensFile used by this component */
+	UFUNCTION(BlueprintCallable, Category="Lens Component")
+	void SetLensFile(FLensFilePicker LensFile);
+
+	/** Get the evaluation mode used to evaluate the LensFile */
+	UFUNCTION(BlueprintPure, Category="Lens Component")
+	EFIZEvaluationMode GetFIZEvaluationMode() const;
+
+	/** Set the evaluation mode used to evaluate the LensFile */
+	UFUNCTION(BlueprintCallable, Category="Lens Component", meta=(DisplayName="Set FIZ Evaluation Mode"))
+	void SetFIZEvaluationMode(EFIZEvaluationMode Mode);
+
+	/** Returns true if nodal offset will be automatically applied during this component's tick, false otherwise */
+	UFUNCTION(BlueprintPure, Category="Lens Component")
+	bool ShouldApplyNodalOffsetOnTick() const;
+
+	/** Set whether nodal offset should be automatically applied during this component's tick */
+	UFUNCTION(BlueprintCallable, Category="Lens Component")
+	void SetApplyNodalOffsetOnTick(bool bApplyNodalOffset);
+
+	/** 
+	 * Manually apply nodal offset to the specified component. 
+	 * If bUseManualInputs is true, the input Focus and Zoom values will be used to evaluate the LensFile .
+	 * If bUseManualInputs is false, the LensFile be will evaluated based on the Lens Component's evaluation mode.
+	 */
+	UFUNCTION(BlueprintCallable, Category="Lens Component", meta=(AdvancedDisplay=1))
+	void ApplyNodalOffset(USceneComponent* ComponentToOffset, bool bUseManualInputs = false, float ManualFocusInput = 0.0f, float ManualZoomInput = 0.0f);
+
 protected:
-	/** Verify base transform and apply nodal offset on top of everything else done in tick */
+	UE_DEPRECATED(5.1, "The use of this callback by this class has been deprecated and it is no longer registered. You can register your own delegate with FWorldDelegates::OnWorldPostActorTick")
 	void OnPostActorTick(UWorld* World, ELevelTick TickType, float DeltaSeconds);
 
 private:
+	/** Evaluate the LensFile for nodal offset (using the current evaluation mode) and apply it to the latest component to offset */
+	void ApplyNodalOffset();
+
 	/** If TargetCameraComponent is not set, initialize it to the first CineCameraComponent on the same actor as this component */
 	void InitDefaultCamera();
 
@@ -49,22 +101,64 @@ private:
 	/** Register a new lens distortion handler with the camera calibration subsystem using the selected lens file */
 	void CreateDistortionHandlerForLensFile();
 
-	/** Evaluate the input LensFile and apply the nodal offset to the input camera's transform */
-	void ApplyNodalOffset(ULensFile* SelectedLensFile, UCineCameraComponent* CineCameraComponent);
+	/** Register to the new LiveLink component's callback to be notified when its controller map changes */
+	void OnLiveLinkComponentRegistered(ULiveLinkComponentController* LiveLinkComponent);
+	
+	/** Callback executed when a LiveLink component on the same actor ticks */
+	void ProcessLiveLinkData(const ULiveLinkComponentController* const LiveLinkComponent, const FLiveLinkSubjectFrameData& SubjectData);
+
+	/** Inspects the subject data and LiveLink transform controller to determine which component (if any) had tracking data applied to it */
+	void UpdateTrackedComponent(const ULiveLinkComponentController* const LiveLinkComponent, const FLiveLinkSubjectFrameData& SubjectData);
+
+	/** Inspects the subject data and LiveLink camera controller cache the FIZ that was input for the target camera */
+	void UpdateLiveLinkFIZ(const ULiveLinkComponentController* const LiveLinkComponent, const FLiveLinkSubjectFrameData& SubjectData);
+
+	/** Updates the focus and zoom inputs that will be used to evaluate the LensFile based on the evaluation mode */
+	void UpdateLensFileEvaluationInputs();
 
 protected:
+	/** Lens File used to drive distortion with current camera settings */
+	UPROPERTY(EditAnywhere, Category="Lens File", meta=(ShowOnlyInnerProperties))
+	FLensFilePicker LensFilePicker;
+
+	/** Specify how the Lens File should be evaluated */
+	UPROPERTY(EditAnywhere, Category="Lens File")
+	EFIZEvaluationMode EvaluationMode = EFIZEvaluationMode::UseLiveLink;
+
 	/** The CineCameraComponent on which to apply the post-process distortion effect */
-	UPROPERTY(EditInstanceOnly, Category = "Default", meta = (UseComponentPicker, AllowedClasses = "/Script/CinematicCamera.CineCameraComponent"))
+	UPROPERTY(EditInstanceOnly, AdvancedDisplay, Category="Lens File", meta=(UseComponentPicker, AllowedClasses="/Script/CinematicCamera.CineCameraComponent"))
 	FComponentReference TargetCameraComponent;
 
+	/** 
+	 * If checked, nodal offset will be applied automatically when this component ticks. 
+	 * Set to false if nodal offset needs to be manually applied at some other time (via Blueprints).
+	 */
+	UPROPERTY(EditAnywhere, Category="Nodal Offset")
+	bool bApplyNodalOffsetOnTick = true;
+
 	/** Structure used to query the camera calibration subsystem for a lens distortion model handler */
-	UPROPERTY(EditAnywhere, BlueprintReadWrite, Category = "Default")
+	UPROPERTY(EditAnywhere, BlueprintReadWrite, Category="Distortion")
 	FDistortionHandlerPicker DistortionSource;
 
 	/** Whether or not to apply distortion to the target camera component */
-	UPROPERTY(EditAnywhere, BlueprintReadWrite, Category = "Default")
+	UPROPERTY(EditAnywhere, BlueprintReadWrite, Category="Distortion")
 	bool bApplyDistortion = false;
 
+	/**
+	 * Whether to use the specified lens file to drive distortion
+	 * Enabling this will create a new distortion source and automatically set this component to use it
+	 */
+	UPROPERTY(EditAnywhere, AdvancedDisplay, Category="Distortion")
+	bool bEvaluateLensFileForDistortion = false;
+
+	/** Whether to scale the computed overscan by the overscan percentage */
+	UPROPERTY(AdvancedDisplay, BlueprintReadWrite, Category="Distortion")
+	bool bScaleOverscan = false;
+
+	/** The percentage of the computed overscan that should be applied to the target camera */
+	UPROPERTY(EditAnywhere, AdvancedDisplay, BlueprintReadWrite, Category="Distortion", meta=(EditCondition="bScaleOverscan", ClampMin="0.0", ClampMax="2.0"))
+	float OverscanMultiplier = 1.0f;
+
 	/** Whether a distortion effect is currently being applied to the target camera component */
 	UPROPERTY()
 	bool bIsDistortionSetup = false;
@@ -73,13 +167,15 @@ protected:
 	UPROPERTY()
 	float OriginalFocalLength = 35.0f;
 
-	/** Original camera rotation to reset before/after applying nodal offset */
+#if WITH_EDITORONLY_DATA
+	UE_DEPRECATED(5.1, "This property has been deprecated. The LensDistortion component no longer tracks the attached camera's original rotation.")
 	UPROPERTY()
-	FRotator OriginalCameraRotation;
+	FRotator OriginalCameraRotation_DEPRECATED;
 
-	/** Original camera location to reset before/after applying nodal offset */
+	UE_DEPRECATED(5.1, "This property has been deprecated. The LensDistortion component no longer tracks the attached camera's original location.")
 	UPROPERTY()
-	FVector OriginalCameraLocation;
+	FVector OriginalCameraLocation_DEPRECATED;
+#endif //WITH_EDITORONLY_DATA
 
 	/** Cached MID last applied to the target camera */
 	UPROPERTY()
@@ -89,29 +185,6 @@ protected:
 	UPROPERTY(Transient)
 	UCineCameraComponent* LastCameraComponent = nullptr;
 
-	/** 
-	 * Whether to use the specified lens file to drive distortion
-	 * Enabling this will create a new distortion source and automatically set this component to use it
-	 */
-	UPROPERTY(EditAnywhere, AdvancedDisplay, Category = "Default")
-	bool bEvaluateLensFileForDistortion = false;
-
-	/** Whether to apply nodal offset to the target camera */
-	UPROPERTY(EditAnywhere, AdvancedDisplay, Category = "Default", meta = (EditCondition = "bEvaluateLensFileForDistortion"))
-	bool bApplyNodalOffset = true;
-
-	/** Lens File used to drive distortion with current camera settings */
-	UPROPERTY(EditAnywhere, AdvancedDisplay, Category = "Default", meta = (EditCondition="bEvaluateLensFileForDistortion"))
-	FLensFilePicker LensFilePicker;
-
-	/** Whether to scale the computed overscan by the overscan percentage */
-	UPROPERTY(AdvancedDisplay, BlueprintReadWrite, Category = "Default")
-	bool bScaleOverscan = false;
-
-	/** The percentage of the computed overscan that should be applied to the target camera */
-	UPROPERTY(EditAnywhere, AdvancedDisplay, BlueprintReadWrite, Category = "Default", meta = (EditCondition = "bScaleOverscan", ClampMin = "0.0", ClampMax = "2.0"))
-	float OverscanMultiplier = 1.0f;
-
 	/** Distortion handler produced by this component */
 	UPROPERTY(Transient)
 	ULensDistortionModelHandlerBase* ProducedLensDistortionHandler = nullptr;
@@ -120,10 +193,17 @@ protected:
 	UPROPERTY(DuplicateTransient)
 	FGuid DistortionProducerID;
 
+	/** Inputs to LensFile evaluation */
+	UPROPERTY(Interp, VisibleAnywhere, AdvancedDisplay, Category="Lens File")
+	float EvalFocus = 0.0f;
+	
+	UPROPERTY(Interp, VisibleAnywhere, AdvancedDisplay, Category="Lens File")
+	float EvalZoom = 0.0f;
+
 private:
-	/** Cached rotation, used to track changes to the camera transform, which impacts the original rotation used by nodal offset */
-	FRotator LastRotation;
+	/** Scene component that should have nodal offset applied */
+	TWeakObjectPtr<USceneComponent> TrackedComponent;
 
-	/** Cached location, used to track changes to the camera transform, which impacts the original location used by nodal offset */
-	FVector LastLocation;
-};
\ No newline at end of file
+	/** Latest LiveLink FIZ data, used to evaluate the LensFile */
+	FLensFileEvalData LiveLinkFIZ;
+};
diff --git a/Engine/Plugins/VirtualProduction/LiveLinkCamera/Source/LiveLinkCamera/Private/LiveLinkCameraController.cpp b/Engine/Plugins/VirtualProduction/LiveLinkCamera/Source/LiveLinkCamera/Private/LiveLinkCameraController.cpp
--- a/Engine/Plugins/VirtualProduction/LiveLinkCamera/Source/LiveLinkCamera/Private/LiveLinkCameraController.cpp
+++ b/Engine/Plugins/VirtualProduction/LiveLinkCamera/Source/LiveLinkCamera/Private/LiveLinkCameraController.cpp
@@ -3,27 +3,17 @@
 #include "LiveLinkCameraController.h"
 
 #include "Camera/CameraComponent.h"
-#include "Camera/PlayerCameraManager.h"
 #include "CameraCalibrationSubsystem.h"
 #include "CineCameraComponent.h"
 #include "Controllers/LiveLinkTransformController.h"
-#include "Engine/Engine.h"
-#include "Engine/World.h"
-#include "Features/IModularFeatures.h"
-#include "GameFramework/Actor.h"
-#include "GameFramework/PlayerController.h"
-#include "ILiveLinkClient.h"
+#include "LensComponent.h"
 #include "LensFile.h"
-#include "LiveLinkComponentController.h"
 #include "LiveLinkLog.h"
 #include "Logging/LogMacros.h"
 #include "Roles/LiveLinkCameraRole.h"
 #include "Roles/LiveLinkCameraTypes.h"
 #include "UObject/EnterpriseObjectVersion.h"
-
-#if WITH_EDITOR
-#include "Kismet2/ComponentEditorUtils.h"
-#endif
+#include "UObject/UE5MainStreamObjectVersion.h"
 
 DEFINE_LOG_CATEGORY_STATIC(LogLiveLinkCameraController, Log, All);
 
@@ -239,8 +229,6 @@ void ULiveLinkCameraController::SetAttachedComponent(UActorComponent* ActorCompo
 	{	
 		// Initialize the most recent filmback and component transform to the current values of the camera to properly detect changes to this property
 		LastFilmback = CineCameraComponent->Filmback;
-		LastRotation = CineCameraComponent->GetRelativeRotation();
-		LastLocation = CineCameraComponent->GetRelativeLocation();
 
 		ULensFile* SelectedLensFile = LensFilePicker.GetLensFile();
 		if (SelectedLensFile && SelectedLensFile->IsCineCameraCompatible(CineCameraComponent) == false)
@@ -262,21 +250,6 @@ void ULiveLinkCameraController::Cleanup()
 			SubSystem->UnregisterDistortionModelHandler(CineCameraComponent, LensDistortionHandler);
 		}
 	}
-
-	if (PostActorTickHandle.IsValid())
-	{
-		FWorldDelegates::OnWorldPostActorTick.Remove(PostActorTickHandle);
-		PostActorTickHandle.Reset();
-	}
-}
-
-void ULiveLinkCameraController::OnEvaluateRegistered()
-{
-	//Hook up to PostActorTick to handle nodal offset
-	if (!PostActorTickHandle.IsValid())
-	{
-		PostActorTickHandle = FWorldDelegates::OnWorldPostActorTick.AddUObject(this, &ULiveLinkCameraController::OnPostActorTick);
-	}
 }
 
 void ULiveLinkCameraController::PostDuplicate(bool bDuplicateForPIE)
@@ -301,26 +274,7 @@ void ULiveLinkCameraController::PostEditImport()
 
 void ULiveLinkCameraController::SetApplyNodalOffset(bool bInApplyNodalOffset)
 {
-	if (bApplyNodalOffset == bInApplyNodalOffset)
-	{
-		return;
-	}
-
-	bApplyNodalOffset = bInApplyNodalOffset;
-
-	if (UCineCameraComponent* const CineCameraComponent = Cast<UCineCameraComponent>(GetAttachedComponent()))
-	{
-		if (bApplyNodalOffset)
-		{
-			OriginalCameraRotation = CineCameraComponent->GetRelativeRotation();
-			OriginalCameraLocation = CineCameraComponent->GetRelativeLocation();
-		}
-		else
-		{
-			CineCameraComponent->SetRelativeLocation(OriginalCameraLocation);
-			CineCameraComponent->SetRelativeRotation(OriginalCameraRotation.Quaternion());
-		}
-	}
+	// The use of this function has been deprecated.
 }
 
 #if WITH_EDITOR
@@ -351,27 +305,8 @@ void ULiveLinkCameraController::PostEditChangeProperty(struct FPropertyChangedEv
 {
 	const FName PropertyName = (PropertyChangedEvent.Property != NULL) ? PropertyChangedEvent.Property->GetFName() : NAME_None;
 
-	if (PropertyName == GET_MEMBER_NAME_CHECKED(ULiveLinkCameraController, bApplyNodalOffset))
-	{
-		if (UCineCameraComponent* const CineCameraComponent = Cast<UCineCameraComponent>(GetAttachedComponent()))
-		{
-			if (bApplyNodalOffset)
-			{
-				OriginalCameraRotation = CineCameraComponent->GetRelativeRotation();
-				OriginalCameraLocation = CineCameraComponent->GetRelativeLocation();
-				UE_LOG(LogLiveLinkCameraController, Verbose, TEXT("Enabling Nodal offset. Cached Location is '%s' and cached rotation is '%s'")
-						, *OriginalCameraLocation.ToString()
-						, *OriginalCameraRotation.ToString());
-			}
-			else
-			{
-				CineCameraComponent->SetRelativeLocation(OriginalCameraLocation);
-				CineCameraComponent->SetRelativeRotation(OriginalCameraRotation.Quaternion());
-			}
-		}
-	}
-	else if (PropertyName == GET_MEMBER_NAME_CHECKED(FLensFilePicker, LensFile)
-			|| PropertyName == GET_MEMBER_NAME_CHECKED(FLensFilePicker, bUseDefaultLensFile))
+	if (PropertyName == GET_MEMBER_NAME_CHECKED(FLensFilePicker, LensFile)
+		|| PropertyName == GET_MEMBER_NAME_CHECKED(FLensFilePicker, bUseDefaultLensFile))
 	{
 		if (UCineCameraComponent* const CineCameraComponent = Cast<UCineCameraComponent>(GetAttachedComponent()))
 		{
@@ -465,8 +400,8 @@ void ULiveLinkCameraController::ApplyFIZ(ULensFile* LensFile, UCineCameraCompone
 					CineCameraComponent->Filmback.SensorHeight = CineCameraComponent->CurrentFocalLength / FocalLengthInfo.FxFy[1];
 
 					// Update the minimum and maximum focal length of the camera (if needed)
-					CineCameraComponent->LensSettings.MinFocalLength = FMath::Min(CineCameraComponent->LensSettings.MinFocalLength, NewFocalLength);
-					CineCameraComponent->LensSettings.MaxFocalLength = FMath::Max(CineCameraComponent->LensSettings.MaxFocalLength, NewFocalLength);
+					CineCameraComponent->LensSettings.MinFocalLength = FMath::Max(FMath::Min(CineCameraComponent->LensSettings.MinFocalLength, NewFocalLength), 0.01);
+					CineCameraComponent->LensSettings.MaxFocalLength = FMath::Max(FMath::Max(CineCameraComponent->LensSettings.MaxFocalLength, NewFocalLength), 0.01);
 
 					CineCameraComponent->SetCurrentFocalLength(NewFocalLength);
 					bHasValidFocalLength = true;
@@ -476,10 +411,6 @@ void ULiveLinkCameraController::ApplyFIZ(ULensFile* LensFile, UCineCameraCompone
 			//If FocalLength could not be applied and it's streamed in, use LiveLink input directly without affecting filmback
 			if (StaticData->bIsFocalLengthSupported && bHasValidFocalLength == false)
 			{
-				// Update the minimum and maximum focal length of the camera (if needed)
-				CineCameraComponent->LensSettings.MinFocalLength = FMath::Min(CineCameraComponent->LensSettings.MinFocalLength, ZoomValue);
-				CineCameraComponent->LensSettings.MaxFocalLength = FMath::Max(CineCameraComponent->LensSettings.MaxFocalLength, ZoomValue);
-
 				CineCameraComponent->SetCurrentFocalLength(ZoomValue);
 			}
 		}
@@ -528,111 +459,7 @@ void ULiveLinkCameraController::ApplyFIZ(ULensFile* LensFile, UCineCameraCompone
 
 void ULiveLinkCameraController::ApplyNodalOffset(ULensFile* SelectedLensFile, UCineCameraComponent* CineCameraComponent)
 {
-	if (CineCameraComponent)
-	{
-		const FRotator CurrentRotator = CineCameraComponent->GetRelativeRotation();
-		const FVector CurrentTranslation = CineCameraComponent->GetRelativeLocation();
-		
-		//Verify if something was set by user / programmatically and update original values
-		bool bWasTransformChanged = false;
-		if (CurrentRotator != LastRotation)
-		{
-			bWasTransformChanged = true;
-			if (CurrentRotator.Pitch != LastRotation.Pitch)
-			{
-				OriginalCameraRotation.Pitch = CurrentRotator.Pitch;
-			}
-
-			if (CurrentRotator.Yaw != LastRotation.Yaw)
-			{
-				OriginalCameraRotation.Yaw = CurrentRotator.Yaw;
-			}
-
-			if (CurrentRotator.Roll != LastRotation.Roll)
-			{
-				OriginalCameraRotation.Roll = CurrentRotator.Roll;
-			}
-		}
-
-		if (CurrentTranslation != LastLocation)
-		{
-			bWasTransformChanged = true;
-			if (CurrentTranslation.X != LastLocation.X)
-			{
-				OriginalCameraLocation.X = CurrentTranslation.X;
-			}
-			if (CurrentTranslation.Y != LastLocation.Y)
-			{
-				OriginalCameraLocation.Y = CurrentTranslation.Y;
-			}
-			if (CurrentTranslation.Z != LastLocation.Z)
-			{
-				OriginalCameraLocation.Z = CurrentTranslation.Z;
-			}
-		}
-
-		if (bApplyNodalOffset && SelectedLensFile)
-		{
-			FNodalPointOffset Offset;
-
-			SelectedLensFile->EvaluateNodalPointOffset(LensFileEvalData.Input.Focus, LensFileEvalData.Input.Zoom, Offset);
-
-			LensFileEvalData.NodalOffset.bWasApplied = true;
-
-			bWasTransformChanged |= Offset.LocationOffset != LastLocationOffset;
-			bWasTransformChanged |= Offset.RotationOffset != LastRotationOffset;
-
-			if (bWasTransformChanged)
-			{
-				CineCameraComponent->SetRelativeLocation(OriginalCameraLocation);
-				CineCameraComponent->SetRelativeRotation(OriginalCameraRotation.Quaternion());
-				CineCameraComponent->AddLocalOffset(Offset.LocationOffset);
-				CineCameraComponent->AddLocalRotation(Offset.RotationOffset);
-
-				// Update the camera manager's cache to reflect this change to nodal offset
-				const UWorld* const World = GetWorld();
-				if (World)
-				{
-					for (FConstPlayerControllerIterator Iterator = World->GetPlayerControllerIterator(); Iterator; ++Iterator)
-					{
-						if (APlayerController* PlayerController = Iterator->Get())
-						{
-							if (APlayerCameraManager* PlayerCameraManager = PlayerController->PlayerCameraManager)
-							{
-								if (PlayerCameraManager->GetViewTarget() == CineCameraComponent->GetOuter())
-								{
-									FMinimalViewInfo CameraViewInfo = PlayerCameraManager->GetCameraCacheView();
-
-									CameraViewInfo.Location = CineCameraComponent->GetComponentLocation();
-									CameraViewInfo.Rotation = CineCameraComponent->GetComponentRotation();
-
-									PlayerCameraManager->FillCameraCache(CameraViewInfo);
-								}
-							}
-						}
-					}
-				}
-			}
-
-			LastLocationOffset = Offset.LocationOffset;
-			LastRotationOffset = Offset.RotationOffset;
-		}
-		else
-		{
-			LastLocationOffset = FVector::ZeroVector;
-			LastRotationOffset = FQuat(0);
-		}
-
-		LastLocation = CineCameraComponent->GetRelativeLocation();
-		LastRotation = CineCameraComponent->GetRelativeRotation();
-	}
-	else 
-	{
-		LastLocation = FVector::OneVector;
-		LastRotation = FRotator::ZeroRotator;
-		LastLocationOffset = FVector::ZeroVector;
-		LastRotationOffset = FQuat(0);
-	}
+	// The use of this function has been deprecated.
 }
 
 void ULiveLinkCameraController::ApplyDistortion(ULensFile* LensFile, UCineCameraComponent* CineCameraComponent, const FLiveLinkCameraStaticData* StaticData, const FLiveLinkCameraFrameData* FrameData)
@@ -670,16 +497,7 @@ void ULiveLinkCameraController::ApplyDistortion(ULensFile* LensFile, UCineCamera
 
 void ULiveLinkCameraController::OnPostActorTick(UWorld* World, ELevelTick TickType, float DeltaSeconds)
 {
-	if (World == GetWorld())
-	{
-		if (UCineCameraComponent* const CineCameraComponent = Cast<UCineCameraComponent>(GetAttachedComponent()))
-		{
-			if (ULensFile* CurrentLensFile = LensFilePicker.GetLensFile())
-			{
-				ApplyNodalOffset(CurrentLensFile, CineCameraComponent);
-			}
-		}
-	}
+	// The use of this callback by this class has been deprecated
 }
 
 void ULiveLinkCameraController::VerifyFIZWithLensFileTables(ULensFile* LensFile, const FLiveLinkCameraStaticData* StaticData) const
@@ -741,7 +559,7 @@ void ULiveLinkCameraController::VerifyFIZWithLensFileTables(ULensFile* LensFile,
 			}
 		}
 
-		//Verify distortion, image center and nodal offset tables 
+		//Verify distortion and image center tables 
 		{
 			if (StaticData->bIsFocusDistanceSupported == false)
 			{
@@ -758,12 +576,6 @@ void ULiveLinkCameraController::VerifyFIZWithLensFileTables(ULensFile* LensFile,
 					FLiveLinkLog::WarningOnce(NAME_InvalidImageCenterMappingWhenNotStreamed, FakedSubjectKey, TEXT("Potential problem evaluating ImageCenter data using subject '%s' and LensFile '%s'. Focus wasn't streamed in and more than one ImageCenter focus point was found."), *FakedSubjectKey.SubjectName.ToString(), *LensFile->GetName());
 				}
 
-				if (bApplyNodalOffset && LensFile->NodalOffsetTable.GetFocusPointNum() > 0 && LensFile->NodalOffsetTable.GetFocusPointNum() != 1)
-				{
-					static const FName NAME_InvalidNodalOffsetMappingWhenNotStreamed = "LiveLinkCamera_FocusNotStreamedInvalidNodalOffsetMapping";
-					FLiveLinkLog::WarningOnce(NAME_InvalidNodalOffsetMappingWhenNotStreamed, FakedSubjectKey, TEXT("Potential problem evaluating NodalOffset data using subject '%s' and LensFile '%s'. Focus wasn't streamed in and more than one NodalOffset focus point was found."), *FakedSubjectKey.SubjectName.ToString(), *LensFile->GetName());
-				}
-
 				if (LensFile->DataMode == ELensDataMode::STMap && LensFile->STMapTable.GetFocusPointNum() > 0 && LensFile->STMapTable.GetFocusPointNum() != 1)
 				{
 					static const FName NAME_InvalidSTMapMappingWhenNotStreamed = "LiveLinkCamera_FocusNotStreamedInvalidSTMapMapping";
@@ -798,15 +610,6 @@ void ULiveLinkCameraController::VerifyFIZWithLensFileTables(ULensFile* LensFile,
 						static const FName NAME_InvalidImageCenterMappingWhenNotStreamed = "LiveLinkCamera_ZoomNotStreamedInvalidImageCenterMapping";
 						FLiveLinkLog::WarningOnce(NAME_InvalidImageCenterMappingWhenNotStreamed, FakedSubjectKey, TEXT("Potential problem evaluating ImageCenter data using subject '%s' and LensFile '%s'. Zoom wasn't streamed in and more than one ImageCenter zoom point was found."), *FakedSubjectKey.SubjectName.ToString(), *LensFile->GetName());
 					}
-
-					if (bApplyNodalOffset && LensFile->NodalOffsetTable.GetFocusPointNum() > 0)
-					{
-						if (LiveLinkCameraControllerUtils::HasSingleZoomPointInvolved<FNodalOffsetFocusPoint>(LensFileEvalData.Input.Focus, LensFile->NodalOffsetTable.GetFocusPoints()) == false)
-						{
-							static const FName NAME_InvalidNodalOffsetMappingWhenNotStreamed = "LiveLinkCamera_ZoomNotStreamedInvalidNodalOffsetMapping";
-							FLiveLinkLog::WarningOnce(NAME_InvalidNodalOffsetMappingWhenNotStreamed, FakedSubjectKey, TEXT("Potential problem evaluating NodalOffset data using subject '%s' and LensFile '%s'. Zoom wasn't streamed in and more than one NodalOffset zoom point was found."), *FakedSubjectKey.SubjectName.ToString(), *LensFile->GetName());
-						}
-					}
 				}
 			}
 		}
@@ -818,8 +621,8 @@ void ULiveLinkCameraController::PostLoad()
 	Super::PostLoad();
 
 #if WITH_EDITOR
-	const int32 Version = GetLinkerCustomVersion(FEnterpriseObjectVersion::GUID);
-	if (Version < FEnterpriseObjectVersion::LiveLinkControllerSplitPerRole)
+	const int32 EnterpriseVersion = GetLinkerCustomVersion(FEnterpriseObjectVersion::GUID);
+	if (EnterpriseVersion < FEnterpriseObjectVersion::LiveLinkControllerSplitPerRole)
 	{
 		AActor* MyActor = GetOuterActor();
 		if (MyActor)
@@ -851,6 +654,69 @@ void ULiveLinkCameraController::PostLoad()
 			}
 		}
 	}
+
+	const int32 UE5MainVersion = GetLinkerCustomVersion(FUE5MainStreamObjectVersion::GUID);
+	if (UE5MainVersion < FUE5MainStreamObjectVersion::LensComponentNodalOffset)
+	{
+		// LensFile evaluation for nodal offset has moved to the Lens component. In order for existing camera actors to continue functioning as before,
+		// we have to migrate some properties from the camera controller to the lens component. If the owning actor does not have already have a lens 
+		// component then we will instance a new one.
+		AActor* const CameraActor = GetOuterActor();
+		CameraActor->ConditionalPostLoad();
+		if (CameraActor->GetWorld())
+		{
+			TInlineComponentArray<ULensComponent*> LensComponents;
+			CameraActor->GetComponents(LensComponents);
+
+			// Loop through any existing lens components to find those that may already have been evaluating the same LensFile as this camera controller
+			bool bFoundMatchingLensComponent = false;
+			for (ULensComponent* LensComponent : LensComponents)
+			{
+				if (LensComponent->GetLensFile() == LensFilePicker.GetLensFile())
+				{
+					bFoundMatchingLensComponent = true;
+
+					// It is very unlikely that both the camera controller and the existing lens component were both applying nodal offset simultaneously.
+					// Therefore, it is safe to assume that if the camera controller was previously applying nodal offset, the new desired behavior would be for 
+					// the lens component to start applying it. But if the lens component was previously applying nodal offset, then there is no need to update anything.
+					PRAGMA_DISABLE_DEPRECATION_WARNINGS
+					if (bApplyNodalOffset_DEPRECATED)
+					{
+						LensComponent->SetApplyNodalOffsetOnTick(bApplyNodalOffset_DEPRECATED);
+					}
+					PRAGMA_ENABLE_DEPRECATION_WARNINGS
+				}
+			}
+
+			// If no matching existing lens components were found, create a new one and initialize it with the correct LensFile
+			if (!bFoundMatchingLensComponent)
+			{
+				if (LensFilePicker.GetLensFile())
+				{
+					ULensComponent* LensComponent = NewObject<ULensComponent>(CameraActor, MakeUniqueObjectName(this, ULensComponent::StaticClass(), TEXT("Lens")));
+					CameraActor->AddInstanceComponent(LensComponent);
+					LensComponent->RegisterComponent();
+
+					LensComponent->SetLensFile(LensFilePicker);
+					PRAGMA_DISABLE_DEPRECATION_WARNINGS
+					LensComponent->SetApplyNodalOffsetOnTick(bApplyNodalOffset_DEPRECATED);
+					PRAGMA_ENABLE_DEPRECATION_WARNINGS
+				}
+			}
+		}
+
+		// If this controller was previously applying nodal offset, reset the camera component's transform to its original pose
+		if (UCineCameraComponent* const CineCameraComponent = Cast<UCineCameraComponent>(GetAttachedComponent()))
+		{
+			PRAGMA_DISABLE_DEPRECATION_WARNINGS
+			if (bApplyNodalOffset_DEPRECATED)
+			{
+				CineCameraComponent->SetRelativeLocation(OriginalCameraLocation_DEPRECATED);
+				CineCameraComponent->SetRelativeRotation(OriginalCameraRotation_DEPRECATED.Quaternion());
+			}
+			PRAGMA_ENABLE_DEPRECATION_WARNINGS
+		}
+	}
 #endif //WITH_EDITOR
 }
 
diff --git a/Engine/Plugins/VirtualProduction/LiveLinkCamera/Source/LiveLinkCamera/Public/LiveLinkCameraController.h b/Engine/Plugins/VirtualProduction/LiveLinkCamera/Source/LiveLinkCamera/Public/LiveLinkCameraController.h
--- a/Engine/Plugins/VirtualProduction/LiveLinkCamera/Source/LiveLinkCamera/Public/LiveLinkCameraController.h
+++ b/Engine/Plugins/VirtualProduction/LiveLinkCamera/Source/LiveLinkCamera/Public/LiveLinkCameraController.h
@@ -67,7 +67,6 @@ public:
 	virtual TSubclassOf<UActorComponent> GetDesiredComponentClass() const override;
 	virtual void SetAttachedComponent(UActorComponent* ActorComponent) override;
 	virtual void Cleanup() override;
-	virtual void OnEvaluateRegistered() override;
 	//~ End ULiveLinkControllerBase interface
 
 	//~ Begin UObject interface
@@ -83,20 +82,20 @@ public:
 	/** Returns a const reference to input data used to evaluate the lens file */
 	const FLensFileEvalData& GetLensFileEvalDataRef() const;
 
-	/** Enables/disables the application of the nodal offset to the camera component */
+	UE_DEPRECATED(5.1, "This function has been deprecated. Nodal Offset is now applied by the LensComponent")
 	void SetApplyNodalOffset(bool bInApplyNodalOffset);
 
 protected:
 	/** Applies FIZ data coming from LiveLink stream. Lens file is used if encoder mapping is required  */
 	void ApplyFIZ(ULensFile* LensFile, UCineCameraComponent* CineCameraComponent, const FLiveLinkCameraStaticData* StaticData, const FLiveLinkCameraFrameData* FrameData);
 
-	/** Applies nodal offset from lens file for the given Focus/Zoom values of CineCamera */
+	UE_DEPRECATED(5.1, "This function has been deprecated. Nodal Offset is now applied by the LensComponent")
 	void ApplyNodalOffset(ULensFile* LensFile, UCineCameraComponent* CineCameraComponent);
 
 	/** Update distortion state */
 	void ApplyDistortion(ULensFile* LensFile, UCineCameraComponent* CineCameraComponent, const FLiveLinkCameraStaticData* StaticData, const FLiveLinkCameraFrameData* FrameData);
 
-	/** Verify base transform and apply nodal offset on top of everything else done in tick */
+	UE_DEPRECATED(5.1, "The use of this callback by this class has been deprecated and it is no longer registered. You can register your own delegate with FWorldDelegates::OnWorldPostActorTick")
 	void OnPostActorTick(UWorld* World, ELevelTick TickType, float DeltaSeconds);
 
 	/** 
@@ -135,9 +134,11 @@ public:
 	UPROPERTY(EditAnywhere, Category = "Camera Calibration")
 	FCameraFilmbackSettings CroppedFilmback;
 
-	/** Apply nodal offset from lens file if enabled */
-	UPROPERTY(EditAnywhere, BlueprintReadOnly, Category = "Camera Calibration")
-	bool bApplyNodalOffset = true;
+#if WITH_EDITORONLY_DATA
+	UE_DEPRECATED(5.1, "This property has been deprecated. The LiveLink camera controller no longer applies nodal offset.")
+	UPROPERTY()
+	bool bApplyNodalOffset_DEPRECATED = true;
+#endif //WITH_EDITORONLY_DATA
 
 	/** Whether to scale the computed overscan by the overscan percentage */
 	UPROPERTY(BlueprintReadWrite, Category = "Camera Calibration")
@@ -156,13 +157,15 @@ protected:
 	UPROPERTY(DuplicateTransient)
 	FGuid DistortionProducerID;
 
-	/** Original cinecamera component rotation that we set back on when nodal offset isn't applied anymore */
+#if WITH_EDITORONLY_DATA
+	UE_DEPRECATED(5.1, "This property has been deprecated. The LiveLink camera controller no longer applies nodal offset.")
 	UPROPERTY()
-	FRotator OriginalCameraRotation;
+	FRotator OriginalCameraRotation_DEPRECATED;
 
-	/** Original cinecamera component location that we set back on when nodal offset isn't applied anymore */
+	UE_DEPRECATED(5.1, "This property has been deprecated. The LiveLink camera controller no longer applies nodal offset.")
 	UPROPERTY()
-	FVector OriginalCameraLocation;
+	FVector OriginalCameraLocation_DEPRECATED;
+#endif //WITH_EDITORONLY_DATA
 
 	/** Used to control which data from LiveLink is actually applied to camera */
 	UPROPERTY(EditAnywhere, Category = "Settings")
@@ -175,19 +178,12 @@ protected:
 #endif
 
 protected:
-
 	/** Caches the latest inputs to the LensFile evaluation */
 	FLensFileEvalData LensFileEvalData;
 
 private:
-
 	//Last values used to detect changes made by the user and update our original caches
 	FCameraFilmbackSettings LastFilmback;
-	FRotator LastRotation;
-	FVector LastLocation;
-	FVector LastLocationOffset;
-	FQuat LastRotationOffset;
-	double LastLensTableVerificationTimestamp = 0.0;
 
-	FDelegateHandle PostActorTickHandle;
-};
\ No newline at end of file
+	double LastLensTableVerificationTimestamp = 0.0;
+};
diff --git a/Engine/Plugins/VirtualProduction/LiveLinkCamera/Source/LiveLinkCameraRecording/LiveLinkCameraRecording.build.cs b/Engine/Plugins/VirtualProduction/LiveLinkCamera/Source/LiveLinkCameraRecording/LiveLinkCameraRecording.build.cs
--- a/Engine/Plugins/VirtualProduction/LiveLinkCamera/Source/LiveLinkCameraRecording/LiveLinkCameraRecording.build.cs
+++ b/Engine/Plugins/VirtualProduction/LiveLinkCamera/Source/LiveLinkCameraRecording/LiveLinkCameraRecording.build.cs
@@ -30,8 +30,10 @@ public class LiveLinkCameraRecording : ModuleRules
 			new string[]
 			{
 				"CameraCalibrationCore",
+				"CinematicCamera",
 				"Core",
 				"CoreUObject",
+				"Engine",
 				"LevelSequence",
 				"LiveLinkCamera",
 				"LiveLinkComponents",
diff --git a/Engine/Plugins/VirtualProduction/LiveLinkCamera/Source/LiveLinkCameraRecording/Private/MovieSceneLiveLinkCameraControllerSection.cpp b/Engine/Plugins/VirtualProduction/LiveLinkCamera/Source/LiveLinkCameraRecording/Private/MovieSceneLiveLinkCameraControllerSection.cpp
--- a/Engine/Plugins/VirtualProduction/LiveLinkCamera/Source/LiveLinkCameraRecording/Private/MovieSceneLiveLinkCameraControllerSection.cpp
+++ b/Engine/Plugins/VirtualProduction/LiveLinkCamera/Source/LiveLinkCameraRecording/Private/MovieSceneLiveLinkCameraControllerSection.cpp
@@ -3,68 +3,26 @@
 #include "MovieSceneLiveLinkCameraControllerSection.h"
 
 #include "Evaluation/MovieScenePreAnimatedState.h"
+#include "Evaluation/PreAnimatedState/MovieScenePreAnimatedCaptureSource.h"
 #include "IMovieScenePlayer.h"
 #include "LiveLinkCameraController.h"
 #include "LiveLinkComponentController.h"
 #include "LiveLinkControllerBase.h"
 #include "MovieSceneExecutionToken.h"
 #include "Roles/LiveLinkCameraRole.h"
-#include "Evaluation/PreAnimatedState/MovieScenePreAnimatedCaptureSource.h"
-
-struct FPreAnimatedLensFileTokenProducer : IMovieScenePreAnimatedTokenProducer
-{
-	virtual IMovieScenePreAnimatedTokenPtr CacheExistingState(UObject& Object) const
-	{
-		struct FLensFileToken : IMovieScenePreAnimatedToken
-		{
-			FLensFileToken(ULiveLinkCameraController* LiveLinkCameraController)
-			{
-				// Save the original values of these settings
-				bOriginalUseDefaultLensFile = LiveLinkCameraController->LensFilePicker.bUseDefaultLensFile;
-				OriginalLensFile = LiveLinkCameraController->LensFilePicker.GetLensFile();
-			}
-
-			virtual void RestoreState(UObject& Object, const UE::MovieScene::FRestoreStateParams& Params)
-			{
-				if (ULiveLinkCameraController* CameraController = CastChecked<ULiveLinkCameraController>(&Object))
-				{
-					// Restore the original values of these settings
-					CameraController->LensFilePicker.bUseDefaultLensFile = bOriginalUseDefaultLensFile;
-					CameraController->LensFilePicker.LensFile = OriginalLensFile.Get();
-				}
-			}
-			
-			// Original settings of the LiveLinkCameraController to save/restore
-			bool bOriginalUseDefaultLensFile;
-			TWeakObjectPtr<ULensFile> OriginalLensFile;
-		};
-
-		return FLensFileToken(CastChecked<ULiveLinkCameraController>(&Object));
-	}
-	static FMovieSceneAnimTypeID GetAnimTypeID()
-	{
-		return TMovieSceneAnimTypeID<FPreAnimatedLensFileTokenProducer>();
-	}
-};
 
 void UMovieSceneLiveLinkCameraControllerSection::Initialize(ULiveLinkControllerBase* InLiveLinkController)
 {
-	// Duplicate the recorded LensFile asset and save the duplicate in this level sequence section
-	if (ULiveLinkCameraController* CameraController = Cast<ULiveLinkCameraController>(InLiveLinkController))
-	{
-		if (ULensFile* LensFile = CameraController->LensFilePicker.GetLensFile())
-		{
-			const FName DuplicateName = FName(*FString::Format(TEXT("{0}_Cached"), { LensFile->GetFName().ToString() }));
-			CachedLensFile = Cast<ULensFile>(StaticDuplicateObject(LensFile, this, DuplicateName, RF_AllFlags & ~RF_Transient));
-		}
-	}
 }
 
-void UMovieSceneLiveLinkCameraControllerSection::Begin(IMovieScenePlayer* Player, const UE::MovieScene::FEvaluationHookParams& Params) const
+void UMovieSceneLiveLinkCameraControllerSection::Update(IMovieScenePlayer* Player, const UE::MovieScene::FEvaluationHookParams& Params) const
 {
-	FScopedPreAnimatedCaptureSource CaptureSource(&Player->PreAnimatedState, this, Params.SequenceID, EvalOptions.CompletionMode == EMovieSceneCompletionMode::RestoreState);
+	if (!CachedLensFile || !bApplyNodalOffsetFromCachedLensFile)
+	{
+		return;
+	}
 
-	for (TWeakObjectPtr<> BoundObject : Player->FindBoundObjects(Params.ObjectBindingID, Params.SequenceID))
+	for (TWeakObjectPtr<>& BoundObject : Player->FindBoundObjects(Params.ObjectBindingID, Params.SequenceID))
 	{
 		if (ULiveLinkComponentController* LiveLinkComponent = Cast<ULiveLinkComponentController>(BoundObject.Get()))
 		{
@@ -73,18 +31,19 @@ void UMovieSceneLiveLinkCameraControllerSection::Begin(IMovieScenePlayer* Player
 			{
 				if (ULiveLinkCameraController* CameraController = Cast<ULiveLinkCameraController>(*Controller))
 				{
-					Player->SavePreAnimatedState(*CameraController, FPreAnimatedLensFileTokenProducer::GetAnimTypeID(), FPreAnimatedLensFileTokenProducer());
+					UActorComponent* CurrentComponentToControl = CameraController->GetAttachedComponent();
+ 					if (UCineCameraComponent* CineCameraComponent = Cast<UCineCameraComponent>(CurrentComponentToControl))
+ 					{
+						const FLensFileEvalData& LensFileEvalData = CameraController->GetLensFileEvalDataRef();
+
+						FNodalPointOffset Offset;
+						CachedLensFile->EvaluateNodalPointOffset(LensFileEvalData.Input.Focus, LensFileEvalData.Input.Zoom, Offset);
 
-					// Override the existing lens file settings of the camera controller to use the cached, duplicate lens file
-					CameraController->LensFilePicker.bUseDefaultLensFile = false;
-					CameraController->LensFilePicker.LensFile = CachedLensFile;
+						CineCameraComponent->AddLocalOffset(Offset.LocationOffset);
+						CineCameraComponent->AddLocalRotation(Offset.RotationOffset);
+ 					}
 				}
 			}
 		}
 	}
 }
-
-void UMovieSceneLiveLinkCameraControllerSection::End(IMovieScenePlayer* Player, const UE::MovieScene::FEvaluationHookParams& Params) const
-{
-	Player->RestorePreAnimatedState();
-}
diff --git a/Engine/Plugins/VirtualProduction/LiveLinkCamera/Source/LiveLinkCameraRecording/Private/MovieSceneLiveLinkCameraControllerSection.h b/Engine/Plugins/VirtualProduction/LiveLinkCamera/Source/LiveLinkCameraRecording/Private/MovieSceneLiveLinkCameraControllerSection.h
--- a/Engine/Plugins/VirtualProduction/LiveLinkCamera/Source/LiveLinkCameraRecording/Private/MovieSceneLiveLinkCameraControllerSection.h
+++ b/Engine/Plugins/VirtualProduction/LiveLinkCamera/Source/LiveLinkCameraRecording/Private/MovieSceneLiveLinkCameraControllerSection.h
@@ -28,10 +28,14 @@ public:
 	void Initialize(ULiveLinkControllerBase* InLiveLinkController);
 
 	//~ Begin UMovieSceneHookSection interface
-	virtual void Begin(IMovieScenePlayer* Player, const UE::MovieScene::FEvaluationHookParams& Params) const override;
-	virtual void End(IMovieScenePlayer* Player, const UE::MovieScene::FEvaluationHookParams& Params) const override;
+	virtual void Update(IMovieScenePlayer* Player, const UE::MovieScene::FEvaluationHookParams& Params) const override;
 	//~ End UMovieSceneHookSection interface
 
+public:
+	/** If true, the Cached Lens File will be evaluated for nodal offset, and that offset will be applied to the transform of the CameraComponent */
+	UPROPERTY(EditAnywhere, Category="Camera Calibration")
+	bool bApplyNodalOffsetFromCachedLensFile = true;
+
 private:
 	/** Saved duplicate of the LensFile asset used by the recorded LiveLink Camera Controller at the time of recording */
 	UPROPERTY(VisibleAnywhere, Category="Camera Calibration")
diff --git a/Engine/Source/Runtime/Core/Public/UObject/UE5MainStreamObjectVersion.h b/Engine/Source/Runtime/Core/Public/UObject/UE5MainStreamObjectVersion.h
--- a/Engine/Source/Runtime/Core/Public/UObject/UE5MainStreamObjectVersion.h
+++ b/Engine/Source/Runtime/Core/Public/UObject/UE5MainStreamObjectVersion.h
@@ -227,6 +227,9 @@ struct CORE_API FUE5MainStreamObjectVersion
 		// Remove Faces in Triangle Mesh BVH
 		RemoveTriangleMeshBVHFaces,
 
+		// Moving all nodal offset handling to Lens Component
+		LensComponentNodalOffset,
+
 		// -----<new versions can be added above this line>-------------------------------------------------
 		VersionPlusOne,
 		LatestVersion = VersionPlusOne - 1
