[CameraCalibration] Centralize lens distortion evaluation to the Lens Component
Deprecate LensDistortionHandlerPicker structure
Remove LensDistortionHandlers from the CameraCalibrationSubsystem
Update Composure to use a component picker to choose a LensComponent that has distortion data
Remove lens distortion evaluation from LiveLink camera and lens controllers

diff --git a/Engine/Plugins/Compositing/Composure/Source/Composure/Classes/CompositingCaptureBase.h b/Engine/Plugins/Compositing/Composure/Source/Composure/Classes/CompositingCaptureBase.h
--- a/Engine/Plugins/Compositing/Composure/Source/Composure/Classes/CompositingCaptureBase.h
+++ b/Engine/Plugins/Compositing/Composure/Source/Composure/Classes/CompositingCaptureBase.h
@@ -5,9 +5,11 @@
 #include "CompositingElement.h"
 
 #include "CameraCalibrationTypes.h"
+
 #include "CompositingCaptureBase.generated.h"
 
 class USceneCaptureComponent2D;
+class ULensComponent;
 class ULensDistortionModelHandlerBase;
 
 
@@ -29,9 +31,9 @@ protected:
 	UPROPERTY(EditAnywhere, BlueprintReadWrite, Category = "Composure|LensDistortion")
 	bool bApplyDistortion = false;
 
-	/** Structure used to query the camera calibration subsystem for a lens distortion model handler */
-	UPROPERTY(EditAnywhere, BlueprintReadWrite, Category = "Composure|LensDistortion")
-	FDistortionHandlerPicker DistortionSource;
+	/** A component reference (customized) that allows the user to specify a component that this controller should control */
+	UPROPERTY(EditInstanceOnly, Category = "Composure|LensDistortion", meta=(DisplayName="Lens Component", UseComponentPicker, AllowedClasses = "/Script/CameraCalibrationCore.LensComponent", AllowAnyActor))
+	FComponentReference LensComponentPicker;
 
 	/** Value used to augment the FOV of the scene capture to produce a CG image with enough data to distort */
 	UPROPERTY(BlueprintReadOnly, Category = "Composure|LensDistortion")
@@ -45,8 +47,15 @@ protected:
 	UPROPERTY()
 	UMaterialInstanceDynamic* LastDistortionMID = nullptr;
 
-public:
+#if WITH_EDITORONLY_DATA
+PRAGMA_DISABLE_DEPRECATION_WARNINGS
+	UE_DEPRECATED(5.1, "This property has been deprecated. Use the LensComponent picker to choose a lens to use for distortion.")
+	UPROPERTY()
+	FDistortionHandlerPicker DistortionSource_DEPRECATED;
+PRAGMA_ENABLE_DEPRECATION_WARNINGS
+#endif //WITH_EDITORONLY_DATA
 
+public:
 	/** Default constructor */
 	ACompositingCaptureBase();
 
@@ -55,19 +64,22 @@ public:
 	void UpdateDistortion();
 
 	//~ Begin UObject Interface
+	virtual void PostInitProperties() override;
 #if WITH_EDITOR
 	virtual void PostEditChangeProperty(struct FPropertyChangedEvent& PropertyChangedEvent) override;
 #endif	
 	//~ End UObject Interface
 
 public:
-
 	/** Sets whether distortion should be applied or not */
 	void SetApplyDistortion(bool bInApplyDistortion);
 
-	/** Sets which distortion handler to use when bInApplyDistortion is enabled */
+	/** Set the Lens Component this CG layer will use to drive distortion on the scene capture */
+	void SetLens(ULensComponent* InLens);
+
+	UE_DEPRECATED(5.1, "This function has been deprecated. Use SetLens() to set a lens component that this CG layer should use to query for distortion.")
 	void SetDistortionHandler(ULensDistortionModelHandlerBase* InDistortionHandler);
 
-	/** Returns the current distortion handler */
+	UE_DEPRECATED(5.1, "This function has been deprecated. Query the lens component used by this CG layer for its distortion handler instead.")
 	ULensDistortionModelHandlerBase* GetDistortionHandler();
 };
diff --git a/Engine/Plugins/Compositing/Composure/Source/Composure/Private/CompositingCaptureBase.cpp b/Engine/Plugins/Compositing/Composure/Source/Composure/Private/CompositingCaptureBase.cpp
--- a/Engine/Plugins/Compositing/Composure/Source/Composure/Private/CompositingCaptureBase.cpp
+++ b/Engine/Plugins/Compositing/Composure/Source/Composure/Private/CompositingCaptureBase.cpp
@@ -6,9 +6,12 @@
 #include "CineCameraComponent.h"
 #include "Components/SceneCaptureComponent2D.h"
 #include "Engine/Engine.h"
+#include "LensComponent.h"
 #include "LensDistortionModelHandlerBase.h"
 
+DEFINE_LOG_CATEGORY_STATIC(LogCompositingCaptureBase, Log, All);
 
+PRAGMA_DISABLE_DEPRECATION_WARNINGS
 ACompositingCaptureBase::ACompositingCaptureBase()
 {
 	// Create the SceneCapture component and assign its parent to be the RootComponent (the ComposurePostProcessingPassProxy)
@@ -18,6 +21,32 @@ ACompositingCaptureBase::ACompositingCaptureBase()
 	// The SceneCaptureComponent2D default constructor disables TAA, but CG Compositing Elements enable it by default
 	SceneCaptureComponent2D->ShowFlags.TemporalAA = true;
 }
+PRAGMA_ENABLE_DEPRECATION_WARNINGS
+
+void ACompositingCaptureBase::PostInitProperties()
+{
+	Super::PostInitProperties();
+	EObjectFlags ExcludeFlags = RF_ArchetypeObject | RF_ClassDefaultObject | RF_NeedLoad | RF_NeedPostLoad | RF_NeedPostLoadSubobjects | RF_WasLoaded;
+	if (!HasAnyFlags(ExcludeFlags))
+	{
+		ACameraActor* TargetCamera = FindTargetCamera();
+
+		if (TargetCamera)
+		{
+			ULensComponent* CurrentLensComponent = Cast<ULensComponent>(LensComponentPicker.GetComponent(TargetCamera));
+			if (!CurrentLensComponent)
+			{
+				TInlineComponentArray<ULensComponent*> LensComponents;
+				TargetCamera->GetComponents(LensComponents);
+
+				if (LensComponents.Num() > 0)
+				{
+					SetLens(LensComponents[0]);
+				}
+			}
+		}
+	}
+}
 
 void ACompositingCaptureBase::UpdateDistortion()
 {
@@ -27,24 +56,21 @@ void ACompositingCaptureBase::UpdateDistortion()
 	{
 		return;
 	}
-	
+
  	if (UCineCameraComponent* const CineCameraComponent = Cast<UCineCameraComponent>(TargetCamera->GetCameraComponent()))
 	{
-		DistortionSource.TargetCameraComponent = CineCameraComponent;
-
 		// Query the camera calibration subsystem for a handler associated with the TargetCamera and matching the user selection
 		ULensDistortionModelHandlerBase* LensDistortionHandler = nullptr;
-		UCameraCalibrationSubsystem* SubSystem = GEngine->GetEngineSubsystem<UCameraCalibrationSubsystem>();
-		if (SubSystem)
-		{
-			LensDistortionHandler = SubSystem->FindDistortionModelHandler(DistortionSource);
 
-			// Get the focal length of the target camera (before any overscan could have been applied) to prevent double overscan
-			const bool bCouldBeDistorted = SubSystem->GetOriginalFocalLength(CineCameraComponent, OriginalFocalLength);
-			if (!bCouldBeDistorted)
-			{
-				OriginalFocalLength = CineCameraComponent->CurrentFocalLength;
-			}
+		ULensComponent* LensComponent = Cast<ULensComponent>(LensComponentPicker.GetComponent(TargetCamera));
+		if (LensComponent && LensComponent->GetOwner() == TargetCamera)
+		{
+			OriginalFocalLength = LensComponent->GetOriginalFocalLength();
+			LensDistortionHandler = LensComponent->GetLensDistortionHandler();
+		}
+		else
+		{
+			OriginalFocalLength = CineCameraComponent->CurrentFocalLength;
 		}
 
 		if (LensDistortionHandler)
@@ -65,7 +91,7 @@ void ACompositingCaptureBase::UpdateDistortion()
 			// If distortion should be applied, add/update the distortion MID to the scene capture's post process materials. Otherwise, remove it.
 			if (SceneCaptureComponent2D)
 			{
-				if (bApplyDistortion)
+				if (bApplyDistortion && LensComponent->WasDistortionEvaluated())
 				{
 					OverscanFactor = LensDistortionHandler->GetOverscanFactor();
 					SceneCaptureComponent2D->AddOrUpdateBlendable(NewDistortionMID);
@@ -95,29 +121,8 @@ void ACompositingCaptureBase::PostEditChangeProperty(struct FPropertyChangedEven
 {
 	const FName PropertyName = (PropertyChangedEvent.Property != NULL) ? PropertyChangedEvent.Property->GetFName() : NAME_None;
 
-	if ((PropertyName == GET_MEMBER_NAME_CHECKED(ACompositingCaptureBase, TargetCameraActor)))
+	if (PropertyName == GET_MEMBER_NAME_CHECKED(ACompositingCaptureBase, TargetCameraActor))
 	{
-		if (TargetCameraActor)
-		{
-			if (UCineCameraComponent* const CineCameraComponent = Cast<UCineCameraComponent>(TargetCameraActor->GetCameraComponent()))
-			{
-				DistortionSource.TargetCameraComponent = CineCameraComponent;
-
-				ULensDistortionModelHandlerBase* LensDistortionHandler = nullptr;
-				UCameraCalibrationSubsystem* SubSystem = GEngine->GetEngineSubsystem<UCameraCalibrationSubsystem>();
-				if (SubSystem)
-				{
-					LensDistortionHandler = SubSystem->FindDistortionModelHandler(DistortionSource);
-				}
-
-				if (LensDistortionHandler == nullptr)
-				{
-					DistortionSource.DistortionProducerID.Invalidate();
-					DistortionSource.HandlerDisplayName.Empty();
-				}
-			}
-		}
-
 		// If there is no target camera, remove the last distortion post-process MID from the scene capture
 		if (TargetCameraActor == nullptr)
 		{
@@ -131,6 +136,20 @@ void ACompositingCaptureBase::PostEditChangeProperty(struct FPropertyChangedEven
 			return;
 		}
 	}
+	else if (PropertyName == GET_MEMBER_NAME_CHECKED(ACompositingCaptureBase, LensComponentPicker))
+	{
+		if (ACameraActor* TargetCamera = FindTargetCamera())
+		{
+			ULensComponent* LensComponent = Cast<ULensComponent>(LensComponentPicker.GetComponent(TargetCamera));
+			if (LensComponent && LensComponent->GetOwner() != TargetCamera)
+			{
+				UE_LOG(LogCompositingCaptureBase, Warning, TEXT("Lens Component '%s' is not attached to the target Camera Actor '%s' for the CG layer '%s'.")
+					, *LensComponent->GetReadableName()
+					, *TargetCamera->GetActorLabel(false)
+					, *this->GetActorLabel(false));
+			}
+		}
+	}
 
 	Super::PostEditChangeProperty(PropertyChangedEvent);
 }
@@ -142,30 +161,32 @@ void ACompositingCaptureBase::SetApplyDistortion(bool bInApplyDistortion)
 	UpdateDistortion();
 }
 
-void ACompositingCaptureBase::SetDistortionHandler(ULensDistortionModelHandlerBase* InDistortionHandler)
+void ACompositingCaptureBase::SetLens(ULensComponent* InLens)
 {
-	if (InDistortionHandler)
-	{
-		DistortionSource.DistortionProducerID = InDistortionHandler->GetDistortionProducerID();
-		DistortionSource.HandlerDisplayName = InDistortionHandler->GetDisplayName();
-	}
-	else
+	if (InLens)
 	{
-		DistortionSource.DistortionProducerID = FGuid();
-		DistortionSource.HandlerDisplayName = FString();
+		LensComponentPicker.OtherActor = InLens->GetOwner();
+		LensComponentPicker.PathToComponent = InLens->GetPathName(InLens->GetOwner());
 	}
+}
 
-	UpdateDistortion();
+void ACompositingCaptureBase::SetDistortionHandler(ULensDistortionModelHandlerBase* InDistortionHandler)
+{
+	// This function has been deprecated.
 }
 
 ULensDistortionModelHandlerBase* ACompositingCaptureBase::GetDistortionHandler()
 {
-	UCameraCalibrationSubsystem* SubSystem = GEngine->GetEngineSubsystem<UCameraCalibrationSubsystem>();
+	ULensDistortionModelHandlerBase* LensDistortionHandler = nullptr;
 
-	if (!SubSystem)
+	if (ACameraActor* TargetCamera = FindTargetCamera())
 	{
-		return nullptr;
+		ULensComponent* LensComponent = Cast<ULensComponent>(LensComponentPicker.GetComponent(TargetCamera));
+		if (LensComponent && LensComponent->GetOwner() == TargetCamera)
+		{
+			LensDistortionHandler = LensComponent->GetLensDistortionHandler();
+		}
 	}
 
-	return SubSystem->FindDistortionModelHandler(DistortionSource);
+	return LensDistortionHandler;
 }
diff --git a/Engine/Plugins/VirtualProduction/CameraCalibration/Source/CameraCalibrationEditor/Private/AssetEditor/CameraCalibrationStepsController.cpp b/Engine/Plugins/VirtualProduction/CameraCalibration/Source/CameraCalibrationEditor/Private/AssetEditor/CameraCalibrationStepsController.cpp
--- a/Engine/Plugins/VirtualProduction/CameraCalibration/Source/CameraCalibrationEditor/Private/AssetEditor/CameraCalibrationStepsController.cpp
+++ b/Engine/Plugins/VirtualProduction/CameraCalibration/Source/CameraCalibrationEditor/Private/AssetEditor/CameraCalibrationStepsController.cpp
@@ -27,6 +27,7 @@
 #include "ICompElementManager.h"
 #include "Input/Events.h"
 #include "Kismet/KismetRenderingLibrary.h"
+#include "LensComponent.h"
 #include "LensFile.h"
 #include "LiveLinkCameraController.h"
 #include "LiveLinkComponentController.h"
@@ -831,39 +832,6 @@ void FCameraCalibrationStepsController::EnableDistortionInCG()
 		return;
 	}
 
-	UCameraCalibrationSubsystem* SubSystem = GEngine->GetEngineSubsystem<UCameraCalibrationSubsystem>();
-
-	if (!SubSystem)
-	{
-		UE_LOG(LogCameraCalibrationEditor, Error, TEXT("Could not find UCameraCalibrationSubsystem"));
-		return;
-	}
-
-	ACineCameraActor* CineCamera = Cast<ACineCameraActor>(Camera.Get());
-
-	if (!CineCamera)
-	{
-		UE_LOG(LogCameraCalibrationEditor, Warning, TEXT("No cine camera selected when trying to enable distortion."));
-		return;
-	}
-
-	// Pick first valid handler that the subsystem finds in our	camera
-	//
-
-	ULensDistortionModelHandlerBase* DistortionHandler = nullptr;
-
-	for (ULensDistortionModelHandlerBase* Handler : SubSystem->GetDistortionModelHandlers(CineCamera->GetCineCameraComponent()))
-	{
-		if (!Handler)
-		{
-			continue;
-		}
-
-		DistortionHandler = Handler;
-
-		break;
-	}
-
 	for (ACompositingElement* Element : Comp->GetChildElements())
 	{
 		ACompositingCaptureBase* CaptureBase = Cast<ACompositingCaptureBase>(Element);
@@ -875,18 +843,6 @@ void FCameraCalibrationStepsController::EnableDistortionInCG()
 
 		// Enable distortion on the CG compositing layer
 		CaptureBase->SetApplyDistortion(true);
-
-		// If a distortion handler exists for the target camera, set it on the CG layer. 
-		// If no handlers currently exist, log a warning. At some later time, if a distortion source is created
-		// the CG layer will automatically pick it up and start using it.
-		if (DistortionHandler)
-		{
-			CaptureBase->SetDistortionHandler(DistortionHandler);
-		}
-		else
-		{
-			UE_LOG(LogCameraCalibrationEditor, Warning, TEXT("Could not find a distortion handler in the selected camera"));
-		}
 	}
 }
 
@@ -1064,21 +1020,37 @@ ULensFile* FCameraCalibrationStepsController::GetLensFile() const
 	return nullptr;
 }
 
-const ULensDistortionModelHandlerBase* FCameraCalibrationStepsController::GetDistortionHandler() const
+ULensComponent* FCameraCalibrationStepsController::FindLensComponent() const
 {
-	if (!CGLayer.IsValid())
+	ACineCameraActor* CineCamera = Cast<ACineCameraActor>(GetCamera());
+	const ULensFile* OpenLensFile = GetLensFile();
+	if (CineCamera && OpenLensFile)
 	{
-		return nullptr;
+		TInlineComponentArray<ULensComponent*> LensComponents;
+		CineCamera->GetComponents(LensComponents);
+
+		const FText TitleError = LOCTEXT("CameraCalibrationEditorError", "Lens File Error");
+
+		for (ULensComponent* LensComponent : LensComponents)
+		{
+			if (LensComponent->GetLensFile() == OpenLensFile)
+			{
+				return LensComponent;
+			}
+		}
 	}
 
-	ACompositingCaptureBase* CaptureBase = Cast<ACompositingCaptureBase>(CGLayer.Get());
+	return nullptr;
+}
 
-	if (!CaptureBase)
+const ULensDistortionModelHandlerBase* FCameraCalibrationStepsController::GetDistortionHandler() const
+{
+	if (ULensComponent* LensComponent = FindLensComponent())
 	{
-		return nullptr;
+		return LensComponent->GetLensDistortionHandler();
 	}
 
-	return CaptureBase->GetDistortionHandler();
+	return nullptr;
 }
 
 bool FCameraCalibrationStepsController::SetMediaSourceUrl(const FString& InMediaSourceUrl)
diff --git a/Engine/Plugins/VirtualProduction/CameraCalibration/Source/CameraCalibrationEditor/Private/AssetEditor/CameraCalibrationStepsController.h b/Engine/Plugins/VirtualProduction/CameraCalibration/Source/CameraCalibrationEditor/Private/AssetEditor/CameraCalibrationStepsController.h
--- a/Engine/Plugins/VirtualProduction/CameraCalibration/Source/CameraCalibrationEditor/Private/AssetEditor/CameraCalibrationStepsController.h
+++ b/Engine/Plugins/VirtualProduction/CameraCalibration/Source/CameraCalibrationEditor/Private/AssetEditor/CameraCalibrationStepsController.h
@@ -18,6 +18,7 @@ class UMediaPlayer;
 class UMediaTexture;
 class UCameraCalibrationStep;
 class UCompositingElementMaterialPass;
+class ULensComponent;
 class ULensDistortionModelHandlerBase;
 class ULensFile;
 class ULiveLinkCameraController;
@@ -94,6 +95,9 @@ public:
 	/** Returns the LensFile that this tool is using */
 	ULensFile* GetLensFile() const;
 
+	/** Returns the first LensComponent attached to the CG camera whose LensFile matches the open asset */
+	ULensComponent* FindLensComponent() const;
+
 	/** Returns the distortion handler used to distort the CG being displayed in the simulcam viewport */
 	const ULensDistortionModelHandlerBase* GetDistortionHandler() const;
 
diff --git a/Engine/Plugins/VirtualProduction/CameraCalibration/Source/CameraCalibrationEditor/Private/Calibrators/CameraNodalOffsetAlgoPoints.cpp b/Engine/Plugins/VirtualProduction/CameraCalibration/Source/CameraCalibrationEditor/Private/Calibrators/CameraNodalOffsetAlgoPoints.cpp
--- a/Engine/Plugins/VirtualProduction/CameraCalibration/Source/CameraCalibrationEditor/Private/Calibrators/CameraNodalOffsetAlgoPoints.cpp
+++ b/Engine/Plugins/VirtualProduction/CameraCalibration/Source/CameraCalibrationEditor/Private/Calibrators/CameraNodalOffsetAlgoPoints.cpp
@@ -586,9 +586,8 @@ void UCameraNodalOffsetAlgoPoints::Tick(float DeltaTime)
 
 			const FLensFileEvalData* LensFileEvalData = StepsController->GetLensFileEvalData();
 
-			// We require lens evaluation data, and that distortion was evaluated so that 2d correlations are valid
-			// Note: The comp enforces distortion application.
-			if (!LensFileEvalData || !LensFileEvalData->Distortion.bWasEvaluated)
+			// We require lens evaluation data
+			if (!LensFileEvalData)
 			{
 				break;
 			}
@@ -608,10 +607,16 @@ void UCameraNodalOffsetAlgoPoints::Tick(float DeltaTime)
 			}
 
 			LastCameraData.Pose = CameraComponent->GetComponentToWorld();
-			LastCameraData.bWasNodalOffsetApplied = DoesCameraPoseIncludeNodalOffset(Camera);
 			LastCameraData.UniqueId = Camera->GetUniqueID();
 			LastCameraData.LensFileEvalData = *LensFileEvalData;
 
+			const ULensComponent* LensComponent = StepsController->FindLensComponent();
+			if (LensComponent)
+			{
+				LastCameraData.bWasNodalOffsetApplied = LensComponent->WasNodalOffsetAppliedThisTick();
+				LastCameraData.bWasDistortionEvaluated = LensComponent->WasDistortionEvaluated();
+			}
+
 			const AActor* CameraParentActor = Camera->GetAttachParentActor();
 
 			if (CameraParentActor)
@@ -854,14 +859,9 @@ bool UCameraNodalOffsetAlgoPoints::ValidateNewRow(TSharedPtr<FCalibrationRowData
 		return false;
 	}
 
-	if (!CalibrationRows.Num())
-	{
-		return true;
-	}
-
 	// Distortion was evaluated
 
-	if (!Row->CameraData.LensFileEvalData.Distortion.bWasEvaluated)
+	if (!Row->CameraData.bWasDistortionEvaluated)
 	{
 		OutErrorMessage = LOCTEXT("DistortionNotEvaluated", "Distortion was not evaluated");
 		return false;
@@ -869,8 +869,6 @@ bool UCameraNodalOffsetAlgoPoints::ValidateNewRow(TSharedPtr<FCalibrationRowData
 
 	// Same LensFile
 
-	const TSharedPtr<FCalibrationRowData>& FirstRow = CalibrationRows[0];
-
 	if (Row->CameraData.LensFileEvalData.LensFile != LensFile)
 	{
 		OutErrorMessage = LOCTEXT("LensFileWasNotTheSame", "Lens file was not the same");
@@ -893,7 +891,13 @@ bool UCameraNodalOffsetAlgoPoints::ValidateNewRow(TSharedPtr<FCalibrationRowData
 		return false;
 	}
 
+	if (!CalibrationRows.Num())
+	{
+		return true;
+	}
+
 	// Same camera as before
+	const TSharedPtr<FCalibrationRowData>& FirstRow = CalibrationRows[0];
 
 	if (FirstRow->CameraData.UniqueId != Row->CameraData.UniqueId)
 	{
@@ -1615,38 +1619,6 @@ bool UCameraNodalOffsetAlgoPoints::IsCalibrationComponentSelected(const UCalibra
 	return ActiveCalibratorComponents.Contains(SelectedComponent);
 }
 
-bool UCameraNodalOffsetAlgoPoints::DoesCameraPoseIncludeNodalOffset(const ACameraActor* Camera) const
-{
-	bool bResult = false;
-
-	if (const FCameraCalibrationStepsController* StepsController = NodalOffsetTool->GetCameraCalibrationStepsController())
-	{
-		const ULensFile* LensFile = StepsController->GetLensFile();
-
-		if (LensFile)
-		{
-			TInlineComponentArray<ULensComponent*> LensComponents;
-			Camera->GetComponents(LensComponents);
-
-			for (const ULensComponent* LensComponent : LensComponents)
-			{
-				if (LensComponent->GetLensFile() == LensFile)
-				{
-					bResult = LensComponent->WasNodalOffsetAppliedThisTick();
-				}
-				else if (LensComponent->WasNodalOffsetAppliedThisTick())
-				{
-					UE_LOG(LogCameraCalibrationEditor, VeryVerbose,
-						TEXT("The camera pose includes a nodal offset from a different lens file than the one currently being calibrated." \
-							"If you see this warning, you should disable the application of nodal offset from other lens components before continuing calibration."))
-				}
-			}
-		}
-	}
-
-	return bResult;
-}
-
 bool UCameraNodalOffsetAlgoPoints::ApplyNodalOffsetToCalibrator()
 {
 	using namespace CameraNodalOffsetAlgoPoints;
diff --git a/Engine/Plugins/VirtualProduction/CameraCalibration/Source/CameraCalibrationEditor/Private/Calibrators/CameraNodalOffsetAlgoPoints.h b/Engine/Plugins/VirtualProduction/CameraCalibration/Source/CameraCalibrationEditor/Private/Calibrators/CameraNodalOffsetAlgoPoints.h
--- a/Engine/Plugins/VirtualProduction/CameraCalibration/Source/CameraCalibrationEditor/Private/Calibrators/CameraNodalOffsetAlgoPoints.h
+++ b/Engine/Plugins/VirtualProduction/CameraCalibration/Source/CameraCalibrationEditor/Private/Calibrators/CameraNodalOffsetAlgoPoints.h
@@ -97,6 +97,9 @@ protected:
 		// True if the cached camera pose included a nodal offset
 		bool bWasNodalOffsetApplied;
 
+		// True if distortion was evaluated (using this LensFile) 
+		bool bWasDistortionEvaluated;
+
 		// The data used to evaluate the lens data in the camera for this sample
 		FLensFileEvalData LensFileEvalData;
 
@@ -232,9 +235,6 @@ protected:
 	/** Returns true if the input component is in the set of active calibration components */
 	bool IsCalibrationComponentSelected(const UCalibrationPointComponent* const SelectedComponent) const;
 
-	/** Check whether the input camera has a lens component that is currently applying a nodal offset to the camera pose */
-	bool DoesCameraPoseIncludeNodalOffset(const ACameraActor* Camera) const;
-
 	/** Update any calibration components in the set of active components if they were replaced by an in-editor event */
 	void OnObjectsReplaced(const TMap<UObject*, UObject*>& OldToNewInstanceMap);
 
diff --git a/Engine/Plugins/VirtualProduction/CameraCalibrationCore/Source/CameraCalibrationCore/Private/CameraCalibrationSubsystem.cpp b/Engine/Plugins/VirtualProduction/CameraCalibrationCore/Source/CameraCalibrationCore/Private/CameraCalibrationSubsystem.cpp
--- a/Engine/Plugins/VirtualProduction/CameraCalibrationCore/Source/CameraCalibrationCore/Private/CameraCalibrationSubsystem.cpp
+++ b/Engine/Plugins/VirtualProduction/CameraCalibrationCore/Source/CameraCalibrationCore/Private/CameraCalibrationSubsystem.cpp
@@ -6,6 +6,7 @@
 #include "CameraCalibrationSettings.h"
 #include "CineCameraComponent.h"
 #include "Engine/TimecodeProvider.h"
+#include "LensComponent.h"
 #include "Misc/CoreDelegates.h"
 #include "UObject/UObjectIterator.h"
 
@@ -42,57 +43,44 @@ ULensFile* UCameraCalibrationSubsystem::GetLensFile(const FLensFilePicker& Picke
 
 TArray<ULensDistortionModelHandlerBase*> UCameraCalibrationSubsystem::GetDistortionModelHandlers(UCineCameraComponent* Component)
 {
+	// This function has been deprecated. The implementation has been changed to provide some backwards compatibility, but code should be updated to not call this function.
 	TArray<ULensDistortionModelHandlerBase*> Handlers;
-	LensDistortionHandlerMap.MultiFind(Component, Handlers);
+
+	if (Component)
+	{
+		AActor* OwningActor = Component->GetOwner();
+
+		TInlineComponentArray<ULensComponent*> LensComponents;
+		OwningActor->GetComponents(LensComponents);
+
+		for (ULensComponent* LensComponent : LensComponents)
+		{
+			Handlers.Add(LensComponent->GetLensDistortionHandler());
+		}
+	}
 
 	return Handlers;
 }
 
+PRAGMA_DISABLE_DEPRECATION_WARNINGS
 ULensDistortionModelHandlerBase* UCameraCalibrationSubsystem::FindDistortionModelHandler(FDistortionHandlerPicker& DistortionHandlerPicker, bool bUpdatePicker) const
 {
+	// This function has been deprecated. The implementation has been changed to provide some backwards compatibility, but code should be updated to not call this function.
 	TArray<ULensDistortionModelHandlerBase*> Handlers;
-	LensDistortionHandlerMap.MultiFind(DistortionHandlerPicker.TargetCameraComponent, Handlers);
 
-	if (!DistortionHandlerPicker.DistortionProducerID.IsValid() && DistortionHandlerPicker.HandlerDisplayName.IsEmpty())
+	if (!DistortionHandlerPicker.TargetCameraComponent)
 	{
-		if (Handlers.Num() > 0)
-		{
-			if (bUpdatePicker)
-			{
-				DistortionHandlerPicker.DistortionProducerID = Handlers[0]->GetDistortionProducerID();
-				DistortionHandlerPicker.HandlerDisplayName = Handlers[0]->GetDisplayName();
-			}
-
-			return Handlers[0];
-		}
+		return nullptr;
 	}
 
-	// Look through the available handlers to find the one driven by the distortion source's producer
-	for (ULensDistortionModelHandlerBase* Handler : Handlers)
-	{
-		if (Handler->GetDistortionProducerID() == DistortionHandlerPicker.DistortionProducerID)
-		{
-			// Reassign the input handler picker's display name to that of the found handler
-			if (bUpdatePicker)
-			{
-				DistortionHandlerPicker.HandlerDisplayName = Handler->GetDisplayName();
-			}
-			return Handler;
-		}
-	}
+	AActor* OwningActor = DistortionHandlerPicker.TargetCameraComponent->GetOwner();
 
-	// Look through the available handlers to find one with the same name as the input handler picker's display name
-	for (ULensDistortionModelHandlerBase* Handler : Handlers)
+	TInlineComponentArray<ULensComponent*> LensComponents;
+	OwningActor->GetComponents(LensComponents);
+
+	if (LensComponents.Num() > 0)
 	{
-		if (Handler->GetDisplayName() == DistortionHandlerPicker.HandlerDisplayName)
-		{
-			// Reassign the input handler picker's producer to that of the found handler
-			if (bUpdatePicker)
-			{
-				DistortionHandlerPicker.DistortionProducerID = Handler->GetDistortionProducerID();
-			}
-			return Handler;
-		}
+		return LensComponents[0]->GetLensDistortionHandler();
 	}
 
 	return nullptr;
@@ -100,55 +88,25 @@ ULensDistortionModelHandlerBase* UCameraCalibrationSubsystem::FindDistortionMode
 
 ULensDistortionModelHandlerBase* UCameraCalibrationSubsystem::FindOrCreateDistortionModelHandler(FDistortionHandlerPicker& DistortionHandlerPicker, const TSubclassOf<ULensModel> LensModelClass)
 {
-	if (LensModelClass == nullptr)
-	{
-		return nullptr;
-	}
-
-	// Attempt to find a handler associated with the input distortion source
-	if (ULensDistortionModelHandlerBase* Handler = FindDistortionModelHandler(DistortionHandlerPicker, false))
+	// This function has been deprecated. The implementation has been changed to provide some backwards compatibility, but code should be updated to not call this function.
+PRAGMA_DISABLE_DEPRECATION_WARNINGS
+	TArray<ULensDistortionModelHandlerBase*> Handlers = GetDistortionModelHandlers(DistortionHandlerPicker.TargetCameraComponent);
+PRAGMA_ENABLE_DEPRECATION_WARNINGS
+	for (ULensDistortionModelHandlerBase* Handler : Handlers)
 	{
-		// If the existing handler supports the input model, simply return it
-		if (Handler->IsModelSupported(LensModelClass))
+		if (Handler->GetLensModelClass() == LensModelClass)
 		{
-			// The display name may have changed, even if the handler already exists, so update the display name of the existing handler to match the distortion source
-			Handler->SetDisplayName(DistortionHandlerPicker.HandlerDisplayName);
 			return Handler;
 		}
-		else
-		{
-			// If the input distortion source has an existing handler, but model does not match the input model, remove the old handler from the map
-			LensDistortionHandlerMap.Remove(DistortionHandlerPicker.TargetCameraComponent, Handler);
-		}
-	}
-
-	// If no handler exists for the input distortion source, create a new one
-	const TSubclassOf<ULensDistortionModelHandlerBase> HandlerClass = ULensModel::GetHandlerClass(LensModelClass);
-	ULensDistortionModelHandlerBase* NewHandler = nullptr;
-	if(HandlerClass)
-	{
-		NewHandler = NewObject<ULensDistortionModelHandlerBase>(DistortionHandlerPicker.TargetCameraComponent, HandlerClass);
-		NewHandler->SetDistortionProducerID(DistortionHandlerPicker.DistortionProducerID);
-		NewHandler->SetDisplayName(DistortionHandlerPicker.HandlerDisplayName);
-		LensDistortionHandlerMap.Add(FObjectKey(DistortionHandlerPicker.TargetCameraComponent), NewHandler);
-
-		if (!CachedFocalLengthMap.Contains(DistortionHandlerPicker.TargetCameraComponent))
-		{
-			FCachedFocalLength CachedFocalLength = { DistortionHandlerPicker.TargetCameraComponent->CurrentFocalLength, 0.0f };
-			CachedFocalLengthMap.Add(DistortionHandlerPicker.TargetCameraComponent, CachedFocalLength);
-		}
-	}
-	else
-	{
-		UE_LOG(LogCameraCalibrationCore, Verbose, TEXT("Could not create DistortionHandler for LensModel '%s'"), *LensModelClass->GetName());
 	}
 
-	return NewHandler;
+	return nullptr;
 }
+PRAGMA_ENABLE_DEPRECATION_WARNINGS
 
 void UCameraCalibrationSubsystem::UnregisterDistortionModelHandler(UCineCameraComponent* Component, ULensDistortionModelHandlerBase* Handler)
 {
-	LensDistortionHandlerMap.Remove(Component, Handler);
+	// This function has been deprecated.
 }
 
 void UCameraCalibrationSubsystem::RegisterDistortionModel(TSubclassOf<ULensModel> LensModel)
@@ -173,34 +131,17 @@ void UCameraCalibrationSubsystem::UnregisterOverlayMaterial(const FName& Materia
 
 void UCameraCalibrationSubsystem::UpdateOriginalFocalLength(UCineCameraComponent* Component, float InFocalLength)
 {
-	FCachedFocalLength* CachedFocalLength = CachedFocalLengthMap.Find(Component);
-	if (CachedFocalLength)
-	{
-		// If the input focal length matches the focal length of the camera after overscan was applied, do not update the original (prevents double overscan)
-		if (!FMath::IsNearlyEqual(InFocalLength, CachedFocalLength->OverscanFocalLength))
-		{
-			CachedFocalLength->OriginalFocalLength = InFocalLength;
-		}
-	}
+	// This function has been deprecated.
 }
 
 void UCameraCalibrationSubsystem::UpdateOverscanFocalLength(UCineCameraComponent* Component, float InFocalLength)
 {
-	FCachedFocalLength* CachedFocalLength = CachedFocalLengthMap.Find(Component);
-	if (CachedFocalLength)
-	{
-		CachedFocalLength->OverscanFocalLength = InFocalLength;
-	}
+	// This function has been deprecated.
 }
 
 bool UCameraCalibrationSubsystem::GetOriginalFocalLength(UCineCameraComponent* Component, float& OutFocalLength)
 {
-	FCachedFocalLength* CachedFocalLength = CachedFocalLengthMap.Find(Component);
-	if (CachedFocalLength)
-	{
-		OutFocalLength = CachedFocalLength->OriginalFocalLength;
-		return true;
-	}
+	// This function has been deprecated.
 	return false;
 }
 
diff --git a/Engine/Plugins/VirtualProduction/CameraCalibrationCore/Source/CameraCalibrationCore/Private/LensComponent.cpp b/Engine/Plugins/VirtualProduction/CameraCalibrationCore/Source/CameraCalibrationCore/Private/LensComponent.cpp
--- a/Engine/Plugins/VirtualProduction/CameraCalibrationCore/Source/CameraCalibrationCore/Private/LensComponent.cpp
+++ b/Engine/Plugins/VirtualProduction/CameraCalibrationCore/Source/CameraCalibrationCore/Private/LensComponent.cpp
@@ -14,6 +14,7 @@
 
 DEFINE_LOG_CATEGORY_STATIC(LogLensComponent, Log, All);
 
+PRAGMA_DISABLE_DEPRECATION_WARNINGS
 ULensComponent::ULensComponent()
 {
 	PrimaryComponentTick.bCanEverTick = true;
@@ -22,12 +23,8 @@ ULensComponent::ULensComponent()
 	// We also use a tick prerequisite on LiveLink components, so technically this could also use TG_PrePhysics
 	PrimaryComponentTick.TickGroup = ETickingGroup::TG_DuringPhysics;
 	bTickInEditor = true;
-
-	if (!HasAnyFlags(RF_ArchetypeObject | RF_ClassDefaultObject))
-	{
-		DistortionProducerID = FGuid::NewGuid();
-	}
 }
+PRAGMA_ENABLE_DEPRECATION_WARNINGS
 
 void ULensComponent::OnRegister()
 {
@@ -35,6 +32,10 @@ void ULensComponent::OnRegister()
 
 	InitDefaultCamera();
 
+	// The component might already have a serialized LensModel, but might not go through PostLoad (in the case that this is spawned by sequencer)
+	// so attempt to create a distortion handler now for the currently set LensModel. 
+	CreateDistortionHandler();
+
 	// Register a callback to let us know when a new LiveLinkComponent is added to this component's parent actor
 	// This gives us the opportunity to track when the subject representation of that LiveLinkComponent changes
 	ILiveLinkComponentsModule& LiveLinkComponentsModule = FModuleManager::GetModuleChecked<ILiveLinkComponentsModule>(TEXT("LiveLinkComponents"));
@@ -76,65 +77,96 @@ void ULensComponent::OnUnregister()
 
 void ULensComponent::TickComponent(float DeltaTime, ELevelTick TickType, FActorComponentTickFunction* ThisTickFunction)
 {
-	// Get the focus and zoom values needed to evaluate the LensFile this Tick
-	UpdateLensFileEvaluationInputs();
-
-	// Attempt to apply nodal offset, which will only succeed if there is a valid LensFile, evaluation inputs, and component to offset
-	if (bApplyNodalOffsetOnTick)
-	{
-		ApplyNodalOffset();
-	}
-
 	if (UCineCameraComponent* const CineCameraComponent = Cast<UCineCameraComponent>(TargetCameraComponent.GetComponent(GetOwner())))
 	{
-		DistortionSource.TargetCameraComponent = CineCameraComponent;
+		// If the camera's focal length has changed since the last time this component ticked, we need to update the original focal length we are using
+		if (CineCameraComponent->CurrentFocalLength != LastFocalLength)
+		{
+			OriginalFocalLength = CineCameraComponent->CurrentFocalLength;
+		}
 
-		UCameraCalibrationSubsystem* const SubSystem = GEngine->GetEngineSubsystem<UCameraCalibrationSubsystem>();
+		// Get the focus and zoom values needed to evaluate the LensFile this Tick
+		UpdateLensFileEvaluationInputs(CineCameraComponent);
+
+		// Attempt to apply nodal offset, which will only succeed if there is a valid LensFile, evaluation inputs, and component to offset
+		if (bApplyNodalOffsetOnTick)
+		{
+			ApplyNodalOffset();
+		}
 
-		// Evaluate the selected lens file for distortion data using the current settings of the target camera
 		ULensFile* LensFile = LensFilePicker.GetLensFile();
-		if (bEvaluateLensFileForDistortion && LensFile)
+
+		UpdateCameraFilmback(CineCameraComponent);
+
+		if (EvaluationMode != EFIZEvaluationMode::UseRecordedValues)
 		{
-			if (!ProducedLensDistortionHandler)
-			{
-				CreateDistortionHandlerForLensFile();
-			}
+			EvaluateFocalLength(CineCameraComponent);
+		}
 
-			//Go through the lens file to get distortion data based on FIZ
-			//Our handler's displacement map will get updated
-			const FVector2D CurrentSensorDimensions = FVector2D(CineCameraComponent->Filmback.SensorWidth, CineCameraComponent->Filmback.SensorHeight);
+		// Evaluate Distortion
+		bWasDistortionEvaluated = false;
+		TObjectPtr<ULensDistortionModelHandlerBase> LensDistortionHandler = LensDistortionHandlerMap.FindRef(LensModel);
 
-			if (SubSystem)
+		if (LensDistortionHandler)
+		{
+			switch (DistortionStateSource)
 			{
-				SubSystem->GetOriginalFocalLength(CineCameraComponent, OriginalFocalLength);
-			}
+			case EDistortionSource::LensFile:
+			{
+				if (!LensFile || !EvalInputs.bIsValid)
+				{
+					break;
+				}
 
-			LensFile->EvaluateDistortionData(CineCameraComponent->CurrentFocusDistance, OriginalFocalLength, CurrentSensorDimensions, ProducedLensDistortionHandler);
+				LensFile->EvaluateDistortionData(EvalInputs.Focus, EvalInputs.Zoom, FVector2D(EvalInputs.Filmback.SensorWidth, EvalInputs.Filmback.SensorHeight), LensDistortionHandler);
 
-			// Adjust overscan by the overscan multiplier
-			if (bScaleOverscan)
+				DistortionState = LensDistortionHandler->GetCurrentDistortionState();
+
+				// Adjust overscan by the overscan multiplier
+				if (bScaleOverscan)
+				{
+					const float ScaledOverscanFactor = ((LensDistortionHandler->GetOverscanFactor() - 1.0f) * OverscanMultiplier) + 1.0f;
+					LensDistortionHandler->SetOverscanFactor(ScaledOverscanFactor);
+				}
+
+				bWasDistortionEvaluated = true;
+				break;
+			}
+			case EDistortionSource::LiveLinkLensSubject:
 			{
-				const float ScaledOverscanFactor = ((ProducedLensDistortionHandler->GetOverscanFactor() - 1.0f) * OverscanMultiplier) + 1.0f;
-				ProducedLensDistortionHandler->SetOverscanFactor(ScaledOverscanFactor);
+				// This case is used by the customization to make the distortion state properties read-only. Otherwise, behavior is the same as EDistortionSource::Manual.
+				// falls through
 			}
+			case EDistortionSource::Manual:
+			{
+				LensDistortionHandler->SetDistortionState(DistortionState);
+
+				//Recompute overscan factor for the distortion state
+				float OverscanFactor = LensDistortionHandler->ComputeOverscanFactor();
+
+				if (bScaleOverscan)
+				{
+					OverscanFactor = ((OverscanFactor - 1.0f) * OverscanMultiplier) + 1.0f;
+				}
+				LensDistortionHandler->SetOverscanFactor(OverscanFactor);
+
+				//Make sure the displacement map is up to date
+				LensDistortionHandler->ProcessCurrentDistortion();
 
-			// Track changes to the cine camera's focal length for consumers of distortion data
-			if (SubSystem)
+				bWasDistortionEvaluated = true;
+				break;
+			}
+			default:
 			{
-				SubSystem->UpdateOriginalFocalLength(CineCameraComponent, CineCameraComponent->CurrentFocalLength);
+				ensureMsgf(false, TEXT("Unhandled EDistortionSource type. Update this switch statement."));
+				break; // Do nothing
+			}
 			}
 		}
 
 		if (bApplyDistortion)
 		{
-			// Get the lens distortion handler for the target camera and distortion source
-			ULensDistortionModelHandlerBase* LensDistortionHandler = nullptr;
-			if (SubSystem)
-			{
-				LensDistortionHandler = SubSystem->FindDistortionModelHandler(DistortionSource);
-			}
-
-			if (LensDistortionHandler)
+			if (LensDistortionHandler && bWasDistortionEvaluated)
 			{
 				// Get the current distortion MID from the lens distortion handler
 				UMaterialInstanceDynamic* NewDistortionMID = LensDistortionHandler->GetDistortionMID();
@@ -149,19 +181,19 @@ void ULensComponent::TickComponent(float DeltaTime, ELevelTick TickType, FActorC
 				// Cache the latest distortion MID
 				LastDistortionMID = NewDistortionMID;
 
-				SubSystem->GetOriginalFocalLength(CineCameraComponent, OriginalFocalLength);
-
-				// Get the overscan factor and use it to modify the target camera's FOV
-				const float OverscanFactor = LensDistortionHandler->GetOverscanFactor();
-				const float OverscanSensorWidth = CineCameraComponent->Filmback.SensorWidth * OverscanFactor;
-				const float OverscanFOV = FMath::RadiansToDegrees(2.0f * FMath::Atan(OverscanSensorWidth / (2.0f * OriginalFocalLength)));
-				CineCameraComponent->SetFieldOfView(OverscanFOV);
-
-				// Update the minimum and maximum focal length of the camera (if needed)
-				CineCameraComponent->LensSettings.MinFocalLength = FMath::Min(CineCameraComponent->LensSettings.MinFocalLength, CineCameraComponent->CurrentFocalLength);
-				CineCameraComponent->LensSettings.MaxFocalLength = FMath::Max(CineCameraComponent->LensSettings.MaxFocalLength, CineCameraComponent->CurrentFocalLength);
-
-				SubSystem->UpdateOverscanFocalLength(CineCameraComponent, CineCameraComponent->CurrentFocalLength);
+				// If using recorded values, the overscanned FOV will already have been applied
+				if (EvaluationMode != EFIZEvaluationMode::UseRecordedValues)
+				{
+					// Get the overscan factor and use it to modify the target camera's FOV
+					const float OverscanFactor = LensDistortionHandler->GetOverscanFactor();
+					const float OverscanSensorWidth = CineCameraComponent->Filmback.SensorWidth * OverscanFactor;
+					const float OverscanFOV = FMath::RadiansToDegrees(2.0f * FMath::Atan(OverscanSensorWidth / (2.0f * OriginalFocalLength)));
+					CineCameraComponent->SetFieldOfView(OverscanFOV);
+
+					// Update the minimum and maximum focal length of the camera (if needed)
+					CineCameraComponent->LensSettings.MinFocalLength = FMath::Min(CineCameraComponent->LensSettings.MinFocalLength, CineCameraComponent->CurrentFocalLength);
+					CineCameraComponent->LensSettings.MaxFocalLength = FMath::Max(CineCameraComponent->LensSettings.MaxFocalLength, CineCameraComponent->CurrentFocalLength);
+				}
 
 				bIsDistortionSetup = true;
 			}
@@ -170,8 +202,12 @@ void ULensComponent::TickComponent(float DeltaTime, ELevelTick TickType, FActorC
 				CleanupDistortion(CineCameraComponent);
 			}
 		}
+
+		LastFocalLength = CineCameraComponent->CurrentFocalLength;
 	}
 
+	bWasLiveLinkFIZUpdated = false;
+
 	Super::TickComponent(DeltaTime, TickType, ThisTickFunction);
 }
 
@@ -182,11 +218,6 @@ void ULensComponent::DestroyComponent(bool bPromoteChildren /*= false*/)
 	if (UCineCameraComponent* const CineCameraComponent = Cast<UCineCameraComponent>(TargetCameraComponent.GetComponent(GetOwner())))
 	{
 		CleanupDistortion(CineCameraComponent);
-
-		if (UCameraCalibrationSubsystem* SubSystem = GEngine->GetEngineSubsystem<UCameraCalibrationSubsystem>())
-		{
-			SubSystem->UnregisterDistortionModelHandler(CineCameraComponent, ProducedLensDistortionHandler);
-		}
 	}
 }
 
@@ -195,12 +226,6 @@ void ULensComponent::PostDuplicate(bool bDuplicateForPIE)
 	Super::PostDuplicate(bDuplicateForPIE);
 
 	InitDefaultCamera();
-
-	// When this component is duplicated (e.g. for PIE), the duplicated component needs a new unique ID
-	if (!DistortionProducerID.IsValid())
-	{
-		DistortionProducerID = FGuid::NewGuid();
-	}
 }
 
 void ULensComponent::PostEditImport()
@@ -208,18 +233,15 @@ void ULensComponent::PostEditImport()
 	Super::PostEditImport();
 
 	InitDefaultCamera();
-
-	// When this component is duplicated (e.g. copy-paste), the duplicated component needs a new unique ID
-	if (!DistortionProducerID.IsValid())
-	{
-		DistortionProducerID = FGuid::NewGuid();
-	}
 }
 
 void ULensComponent::PostLoad()
 {
 	Super::PostLoad();
 
+	// After LensModel has been loaded, create a new handler
+	CreateDistortionHandler();
+
 	if (bIsDistortionSetup)
 	{
 		if (UCineCameraComponent* const CineCameraComponent = Cast<UCineCameraComponent>(TargetCameraComponent.GetComponent(GetOwner())))
@@ -263,71 +285,23 @@ void ULensComponent::PostEditChangeProperty(struct FPropertyChangedEvent& Proper
 			}
 		}
 	}
-	else if (PropertyName == GET_MEMBER_NAME_CHECKED(ULensComponent, bEvaluateLensFileForDistortion))
-	{
-		if (bEvaluateLensFileForDistortion)
-		{
-			// Create a new distortion handler if one does not yet exist for the selected lens file
-			if (!ProducedLensDistortionHandler)
-			{
-				CreateDistortionHandlerForLensFile();
-			}
-		}
-		else
-		{
-			// Reset the distortion source settings
-			if (DistortionSource.DistortionProducerID == DistortionProducerID)
-			{
-				DistortionSource.HandlerDisplayName.Empty();
-				DistortionSource.DistortionProducerID.Invalidate();
-			}
-
-			// Unregister the produced handler from the subsystem
-			if (UCameraCalibrationSubsystem* SubSystem = GEngine->GetEngineSubsystem<UCameraCalibrationSubsystem>())
-			{
-				if (UCineCameraComponent* const CineCameraComponent = Cast<UCineCameraComponent>(TargetCameraComponent.GetComponent(GetOwner())))
-				{
-					SubSystem->UnregisterDistortionModelHandler(CineCameraComponent, ProducedLensDistortionHandler);
-				}
-			}
-
-			ProducedLensDistortionHandler = nullptr;
-		}
-	}
 	else if (PropertyName == GET_MEMBER_NAME_CHECKED(ULensComponent, TargetCameraComponent))
 	{
 		// Clean up distortion on the last target camera
 		if (LastCameraComponent)
 		{
 			CleanupDistortion(LastCameraComponent);
-
-			// Unregister the produced handler (if one exists) from the subsystem for the last camera component
-			if (ProducedLensDistortionHandler)
-			{
-				if (UCameraCalibrationSubsystem* SubSystem = GEngine->GetEngineSubsystem<UCameraCalibrationSubsystem>())
-				{
-					SubSystem->UnregisterDistortionModelHandler(LastCameraComponent, ProducedLensDistortionHandler);
-				}
-
-				ProducedLensDistortionHandler = nullptr;
-			}
 		}
 
 		LastCameraComponent = Cast<UCineCameraComponent>(TargetCameraComponent.GetComponent(GetOwner()));
-		DistortionSource.TargetCameraComponent = LastCameraComponent;
-
-		// Create a new handler for the new target camera component
-		if (bEvaluateLensFileForDistortion)
-		{
-			CreateDistortionHandlerForLensFile();
-		}
 	}
 	else if (PropertyName == GET_MEMBER_NAME_CHECKED(FLensFilePicker, LensFile))
 	{
-		if (bEvaluateLensFileForDistortion)
-		{
-			CreateDistortionHandlerForLensFile();
-		}
+		SetLensFilePicker(LensFilePicker);
+	}
+	else if (PropertyName == GET_MEMBER_NAME_CHECKED(ULensComponent, LensModel))
+	{
+		SetLensModel(LensModel);
 	}
 
 	Super::PostEditChangeProperty(PropertyChangedEvent);
@@ -339,6 +313,31 @@ void ULensComponent::OnPostActorTick(UWorld* World, ELevelTick TickType, float D
 	// The use of this callback by this class has been deprecated
 }
 
+void ULensComponent::ReapplyNodalOffset()
+{
+	// Find the TrackedComponent based on the serialized name
+	TInlineComponentArray<USceneComponent*> SceneComponents;
+	GetOwner()->GetComponents(SceneComponents);
+
+	for (USceneComponent* SceneComponent : SceneComponents)
+	{
+		if (SceneComponent->GetName() == TrackedComponentName)
+		{
+			TrackedComponent = SceneComponent;
+			break;
+		}
+	}
+
+	if (TrackedComponent.IsValid())
+	{
+		// Reset the tracked component back to its original relative transform (before nodal offset was originally applied)
+		TrackedComponent->SetRelativeTransform(OriginalTrackedComponentTransform);
+
+		// Now, reapply the nodal offset to the tracked component
+		ApplyNodalOffset();
+	}
+}
+
 void ULensComponent::ApplyNodalOffset()
 {
 	// Verify that we detected a tracked component for us to offset this tick. If there is none, it is likely because no LiveLink transform controller executed this tick.
@@ -348,7 +347,7 @@ void ULensComponent::ApplyNodalOffset()
 	}
 
 	// Verify that the evaluation inputs are valid
-	if (EvaluationMode == EFIZEvaluationMode::DoNotEvaluate)
+	if ((EvaluationMode == EFIZEvaluationMode::DoNotEvaluate) || !EvalInputs.bIsValid) 
 	{
 		return;
 	}
@@ -361,9 +360,10 @@ void ULensComponent::ApplyNodalOffset()
 	}
 
 	FNodalPointOffset Offset;
-	LensFile->EvaluateNodalPointOffset(EvalFocus, EvalZoom, Offset);
+	LensFile->EvaluateNodalPointOffset(EvalInputs.Focus, EvalInputs.Zoom, Offset);
 
-	// TODO: This would be a good time to cache the raw tracking pose before applying the nodal offset
+	// Cache the original transform before applying the offset, so that nodal offset can potentially be re-evaluated in the future
+	OriginalTrackedComponentTransform = TrackedComponent.Get()->GetRelativeTransform();
 
 	TrackedComponent.Get()->AddLocalOffset(Offset.LocationOffset);
 	TrackedComponent.Get()->AddLocalRotation(Offset.RotationOffset);
@@ -398,21 +398,93 @@ void ULensComponent::ApplyNodalOffset(USceneComponent* ComponentToOffset, bool b
 	}
 	else
 	{
-		LensFile->EvaluateNodalPointOffset(EvalFocus, EvalZoom, Offset);
+		if (!EvalInputs.bIsValid)
+		{
+			return;
+		}
+
+		LensFile->EvaluateNodalPointOffset(EvalInputs.Focus, EvalInputs.Zoom, Offset);
 	}
 
+	// Cache the original transform before applying the offset, so that nodal offset can potentially be re-evaluated in the future
+	OriginalTrackedComponentTransform = TrackedComponent.Get()->GetRelativeTransform();
+
 	ComponentToOffset->AddLocalOffset(Offset.LocationOffset);
 	ComponentToOffset->AddLocalRotation(Offset.RotationOffset);
 }
 
+void ULensComponent::EvaluateFocalLength(UCineCameraComponent* CineCameraComponent)
+{
+	FFocalLengthInfo FocalLengthInfo;
+	ULensFile* LensFile = LensFilePicker.GetLensFile();
+	if (LensFile && EvalInputs.bIsValid && LensFile->EvaluateFocalLength(EvalInputs.Focus, EvalInputs.Zoom, FocalLengthInfo))
+	{
+		// TODO: EvaluateFocalLength should do this scaling internally
+		const FVector2D FxFyScale = FVector2D(LensFile->LensInfo.SensorDimensions.X / EvalInputs.Filmback.SensorWidth, LensFile->LensInfo.SensorDimensions.Y / EvalInputs.Filmback.SensorHeight);
+		FocalLengthInfo.FxFy = FocalLengthInfo.FxFy * FxFyScale;
+
+		if ((FocalLengthInfo.FxFy.X > UE_KINDA_SMALL_NUMBER) && (FocalLengthInfo.FxFy.Y > UE_KINDA_SMALL_NUMBER))
+		{
+			// This is how field of view, filmback, and focal length are related:
+			//
+			// FOVx = 2*atan(1/(2*Fx)) = 2*atan(FilmbackX / (2*FocalLength))
+			// => FocalLength = Fx*FilmbackX
+			// 
+			// FOVy = 2*atan(1/(2*Fy)) = 2*atan(FilmbackY / (2*FocalLength))
+			// => FilmbackY = FocalLength / Fy
+
+			// Adjust FocalLength and Filmback to match FxFy (which has already been divided by resolution in pixels)
+			const float NewFocalLength = FocalLengthInfo.FxFy.X * EvalInputs.Filmback.SensorWidth;
+
+			// TODO: Consider adding an advanced setting to control whether or not sensor height should be modified (if Fx/Fy don't have a fixed aspect ratio)
+
+			// Update the minimum and maximum focal length of the camera (if needed)
+			CineCameraComponent->LensSettings.MinFocalLength = FMath::Max(FMath::Min(CineCameraComponent->LensSettings.MinFocalLength, NewFocalLength), 0.01);
+			CineCameraComponent->LensSettings.MaxFocalLength = FMath::Max(FMath::Max(CineCameraComponent->LensSettings.MaxFocalLength, NewFocalLength), 0.01);
+
+			CineCameraComponent->SetCurrentFocalLength(NewFocalLength);
+
+			OriginalFocalLength = NewFocalLength;
+		}
+	}
+}
+
+FLensFilePicker ULensComponent::GetLensFilePicker() const
+{
+	return LensFilePicker;
+}
+
 ULensFile* ULensComponent::GetLensFile() const
 {
 	return LensFilePicker.GetLensFile();
 }
 
-void ULensComponent::SetLensFile(FLensFilePicker LensFile)
+void ULensComponent::SetLensFilePicker(FLensFilePicker LensFile)
 {
 	LensFilePicker = LensFile;
+
+	if (DistortionStateSource == EDistortionSource::LensFile)
+	{
+		if (ULensFile* const NewLensFile = LensFilePicker.GetLensFile())
+		{
+			SetLensModel(NewLensFile->LensInfo.LensModel);
+		}
+	}
+}
+
+void ULensComponent::SetLensFile(ULensFile* Lens)
+{
+	// Automatically sets this to false so the component can use the newly set LensFile directly
+	LensFilePicker.bUseDefaultLensFile = false;
+	LensFilePicker.LensFile = Lens;
+
+	if (DistortionStateSource == EDistortionSource::LensFile)
+	{
+		if (Lens)
+		{
+			SetLensModel(Lens->LensInfo.LensModel);
+		}
+	}
 }
 
 EFIZEvaluationMode ULensComponent::GetFIZEvaluationMode() const
@@ -425,6 +497,46 @@ void ULensComponent::SetFIZEvaluationMode(EFIZEvaluationMode Mode)
 	EvaluationMode = Mode;
 }
 
+float ULensComponent::GetOverscanMultiplier() const
+{
+	if (bScaleOverscan)
+	{
+		return OverscanMultiplier;
+	}
+	return 1.0f;
+}
+
+void ULensComponent::SetOverscanMultiplier(float Multiplier)
+{
+	bScaleOverscan = true;
+	OverscanMultiplier = Multiplier;
+}
+
+EFilmbackOverrideSource ULensComponent::GetFilmbackOverrideSetting() const
+{
+	return FilmbackOverride;
+}
+
+void ULensComponent::SetFilmbackOverrideSetting(EFilmbackOverrideSource Setting)
+{
+	FilmbackOverride = Setting;
+}
+
+FCameraFilmbackSettings ULensComponent::GetCroppedFilmback() const
+{
+	return CroppedFilmback;
+}
+
+void ULensComponent::SetCroppedFilmback(FCameraFilmbackSettings Filmback)
+{
+	CroppedFilmback = Filmback;
+}
+
+const FLensFileEvaluationInputs& ULensComponent::GetLensFileEvaluationInputs() const 
+{
+	return EvalInputs;
+}
+
 bool ULensComponent::ShouldApplyNodalOffsetOnTick() const
 {
 	return bApplyNodalOffsetOnTick;
@@ -440,10 +552,84 @@ bool ULensComponent::WasNodalOffsetAppliedThisTick() const
 	return bWasNodalOffsetAppliedThisTick;
 }
 
+bool ULensComponent::WasDistortionEvaluated() const
+{
+	return bWasDistortionEvaluated;
+}
+
+float ULensComponent::GetOriginalFocalLength() const 
+{
+	return OriginalFocalLength;
+}
+
+TSubclassOf<ULensModel> ULensComponent::GetLensModel() const
+{
+	return LensModel;
+}
+
+void ULensComponent::SetLensModel(TSubclassOf<ULensModel> Model)
+{
+	LensModel = Model;
+
+	if (LensModel)
+	{
+		const uint32 NumDistortionParameters = LensModel->GetDefaultObject<ULensModel>()->GetNumParameters();
+		DistortionState.DistortionInfo.Parameters.Init(0.0f, NumDistortionParameters);
+
+		CreateDistortionHandler();
+	}
+}
+
+FLensDistortionState ULensComponent::GetDistortionState() const
+{
+	return DistortionState;
+}
+
+void ULensComponent::SetDistortionState(FLensDistortionState State)
+{
+	DistortionState = State;
+}
+
+EDistortionSource ULensComponent::GetDistortionSource() const
+{
+	return DistortionStateSource;
+}
+
+void ULensComponent::SetDistortionSource(EDistortionSource Source)
+{
+	DistortionStateSource = Source;
+}
+
+bool ULensComponent::ShouldApplyDistortion() const
+{
+	return bApplyDistortion;
+}
+
+void ULensComponent::SetApplyDistortion(bool bApply)
+{
+	bApplyDistortion = bApply;
+}
+
+ULensDistortionModelHandlerBase* ULensComponent::GetLensDistortionHandler() const
+{
+	return LensDistortionHandlerMap.FindRef(LensModel);
+}
+
+PRAGMA_DISABLE_DEPRECATION_WARNINGS
+FDistortionHandlerPicker ULensComponent::GetDistortionHandlerPicker() const
+{
+#if WITH_EDITOR
+	return DistortionSource_DEPRECATED;
+#else
+	return FDistortionHandlerPicker();
+#endif //WITH_EDITOR
+}
+PRAGMA_ENABLE_DEPRECATION_WARNINGS
+
 void ULensComponent::InitDefaultCamera()
 {
-	UCineCameraComponent* const CineCameraComponent = Cast<UCineCameraComponent>(TargetCameraComponent.GetComponent(GetOwner()));
-	if (!CineCameraComponent)
+	UCineCameraComponent* TargetCamera = Cast<UCineCameraComponent>(TargetCameraComponent.GetComponent(GetOwner()));
+	if (!TargetCamera)
 	{
 		// Find all CineCameraComponents on the same actor as this component and set the first one to be the target
 		TInlineComponentArray<UCineCameraComponent*> CineCameraComponents;
@@ -474,36 +660,22 @@ void ULensComponent::CleanupDistortion(UCineCameraComponent* const CineCameraCom
 		// Update the minimum and maximum focal length of the camera (if needed)
 		CineCameraComponent->LensSettings.MinFocalLength = FMath::Min(CineCameraComponent->LensSettings.MinFocalLength, CineCameraComponent->CurrentFocalLength);
 		CineCameraComponent->LensSettings.MaxFocalLength = FMath::Max(CineCameraComponent->LensSettings.MaxFocalLength, CineCameraComponent->CurrentFocalLength);
-
-		UCameraCalibrationSubsystem* SubSystem = GEngine->GetEngineSubsystem<UCameraCalibrationSubsystem>();
-		if (SubSystem)
-		{
-			SubSystem->UpdateOverscanFocalLength(CineCameraComponent, 0.0f);
-		}
 	}
 
 	bIsDistortionSetup = false;
 }
 
-void ULensComponent::CreateDistortionHandlerForLensFile()
+void ULensComponent::CreateDistortionHandler()
 {
-	if (DistortionSource.TargetCameraComponent == nullptr)
+	if (LensModel)
 	{
-		return;
-	}
-
-	ULensFile* const LensFile = LensFilePicker.GetLensFile();
-
-	// Set the distortion source settings to match this component and the selected lens file
-	if (LensFile)
-	{
-		DistortionSource.HandlerDisplayName = FString::Format(TEXT("{0} (Lens File) (Lens Distortion Component)"), { LensFile->GetFName().ToString() });
-		DistortionSource.DistortionProducerID = DistortionProducerID;
-
-		UCameraCalibrationSubsystem* const SubSystem = GEngine->GetEngineSubsystem<UCameraCalibrationSubsystem>();
-		if (SubSystem)
+		TObjectPtr<ULensDistortionModelHandlerBase>& Handler = LensDistortionHandlerMap.FindOrAdd(LensModel);
+		if (!Handler)
 		{
-			ProducedLensDistortionHandler = SubSystem->FindOrCreateDistortionModelHandler(DistortionSource, LensFile->LensInfo.LensModel);
+			if (const TSubclassOf<ULensDistortionModelHandlerBase> HandlerClass = ULensModel::GetHandlerClass(LensModel))
+			{
+				Handler = NewObject<ULensDistortionModelHandlerBase>(this, HandlerClass);
+			}
 		}
 	}
 }
@@ -557,7 +729,10 @@ void ULensComponent::UpdateTrackedComponent(const ULiveLinkComponentController*
 			if (bIsLocationRotationSupported)
 			{
 				TrackedComponent = Cast<USceneComponent>(TransformController->GetAttachedComponent());
-
+				if (TrackedComponent.IsValid())
+				{
+					TrackedComponentName = TrackedComponent.Get()->GetName();
+				}
 				// At this point, we know that the tracked component's transform was just set by LiveLink, and therefore has no nodal offset applied
 				bWasNodalOffsetAppliedThisTick = false;
 			}
@@ -588,28 +763,67 @@ void ULensComponent::UpdateLiveLinkFIZ(const ULiveLinkComponentController* const
 	{
 		LiveLinkFIZ.Input.Focus = FrameData->FocusDistance;
 		LiveLinkFIZ.Input.Zoom = FrameData->FocalLength;
+
+		bWasLiveLinkFIZUpdated = true;
 	}
 }
 
-void ULensComponent::UpdateLensFileEvaluationInputs()
+void ULensComponent::UpdateLensFileEvaluationInputs(UCineCameraComponent* CineCameraComponent)
 {
+	EvalInputs.bIsValid = false;
+
 	// Query for the focus and zoom inputs to use to evaluate the LensFile. The source of these inputs will depend on the evaluation mode.
 	if (EvaluationMode == EFIZEvaluationMode::UseLiveLink)
 	{
-		EvalFocus = LiveLinkFIZ.Input.Focus;
-		EvalZoom = LiveLinkFIZ.Input.Zoom;
+		if (bWasLiveLinkFIZUpdated)
+		{
+			EvalInputs.Focus = LiveLinkFIZ.Input.Focus;
+			EvalInputs.Zoom = LiveLinkFIZ.Input.Zoom;
+			EvalInputs.bIsValid = true;
+		}
 	}
 	else if (EvaluationMode == EFIZEvaluationMode::UseCameraSettings)
 	{
-		if (const UCineCameraComponent* const CineCameraComponent = Cast<UCineCameraComponent>(TargetCameraComponent.GetComponent(GetOwner())))
-		{
-			EvalFocus = CineCameraComponent->CurrentFocusDistance;
-			EvalZoom = CineCameraComponent->CurrentFocalLength;
-		}
+		EvalInputs.Focus = CineCameraComponent->CurrentFocusDistance;
+		// The camera's focal length might have been altered by applying an overscan FOV. However, the LensFile needs to be evaluated with the non-overscanned value
+		EvalInputs.Zoom = OriginalFocalLength;
+		EvalInputs.bIsValid = true;
 	}
 	else if (EvaluationMode == EFIZEvaluationMode::UseRecordedValues)
 	{
-		// Do nothing, the values for EvalFocus and EvalZoom are already loaded from the recorded sequence
+		// Do nothing, the values for EvalInputs.Focus and EvalInputs.Zoom are already loaded from the recorded sequence
+		EvalInputs.bIsValid = true;
 	}
 }
 
+void ULensComponent::UpdateCameraFilmback(UCineCameraComponent* CineCameraComponent)
+{	
+	switch (FilmbackOverride)
+	{
+		case EFilmbackOverrideSource::LensFile:
+		{	
+			if (ULensFile* LensFile = LensFilePicker.GetLensFile())
+			{
+				CineCameraComponent->Filmback.SensorWidth = LensFile->LensInfo.SensorDimensions.X;
+				CineCameraComponent->Filmback.SensorHeight = LensFile->LensInfo.SensorDimensions.Y;
+			}
+			break;
+		}
+		case EFilmbackOverrideSource::CroppedFilmbackSetting:
+		{
+			CineCameraComponent->Filmback = CroppedFilmback;
+			break;
+		}
+		case EFilmbackOverrideSource::DoNotOverride:
+		{
+			break; // Do Nothing
+		}
+		default:
+		{
+			ensureMsgf(false, TEXT("Unhandled EFilmbackOverrideSource type. Update this switch statement."));
+			break; 
+		}
+	}
+
+	EvalInputs.Filmback = CineCameraComponent->Filmback;
+}
diff --git a/Engine/Plugins/VirtualProduction/CameraCalibrationCore/Source/CameraCalibrationCore/Private/LensFile.cpp b/Engine/Plugins/VirtualProduction/CameraCalibrationCore/Source/CameraCalibrationCore/Private/LensFile.cpp
--- a/Engine/Plugins/VirtualProduction/CameraCalibrationCore/Source/CameraCalibrationCore/Private/LensFile.cpp
+++ b/Engine/Plugins/VirtualProduction/CameraCalibrationCore/Source/CameraCalibrationCore/Private/LensFile.cpp
@@ -711,6 +711,16 @@ bool ULensFile::EvaluateDistortionForSTMaps(float InFocus, float InZoom, FVector
 		return true;
 	}
 
+	if ((LensInfo.SensorDimensions.X < InFilmback.X) || (LensInfo.SensorDimensions.Y < InFilmback.Y))
+	{
+		UE_LOG(LogCameraCalibrationCore, Verbose
+			, TEXT("Can't evaluate LensFile '%s' - The filmback used to generate the calibrated ST Maps is smaller than" \
+				"the filmback of the camera that distortion is being applied to. There is not enough distortion information available in the ST Maps.")
+			, *GetName());
+		SetupNoDistortionOutput(InLensHandler);
+		return false;
+	}
+
 	FDisplacementMapBlendingParams Params;
 
 	TArray<UTextureRenderTarget2D*, TInlineAllocator<4>> UndistortionMapSource;
@@ -978,20 +988,6 @@ void ULensFile::OnDistortionDerivedDataJobCompleted(const FDerivedDistortionData
 
 bool ULensFile::IsCineCameraCompatible(const UCineCameraComponent* CineCameraComponent) const
 {
-	if (!CineCameraComponent)
-	{ 
-		return false;
-	}
-
-	if (DataMode == ELensDataMode::STMap)
-	{
-		if ((LensInfo.SensorDimensions.X < CineCameraComponent->Filmback.SensorWidth)
-			|| (LensInfo.SensorDimensions.Y < CineCameraComponent->Filmback.SensorHeight))
-		{
-			return false;
-		}
-	}
-
 	return true;
 }
 
diff --git a/Engine/Plugins/VirtualProduction/CameraCalibrationCore/Source/CameraCalibrationCore/Public/CameraCalibrationSubsystem.h b/Engine/Plugins/VirtualProduction/CameraCalibrationCore/Source/CameraCalibrationCore/Public/CameraCalibrationSubsystem.h
--- a/Engine/Plugins/VirtualProduction/CameraCalibrationCore/Source/CameraCalibrationCore/Public/CameraCalibrationSubsystem.h
+++ b/Engine/Plugins/VirtualProduction/CameraCalibrationCore/Source/CameraCalibrationCore/Public/CameraCalibrationSubsystem.h
@@ -19,8 +19,7 @@
 
 class UCineCameraComponent;
 
-/** Utility struct to store the original and overscanned focal lengths of a camera */
-struct FCachedFocalLength
+struct UE_DEPRECATED(5.1, "This struct has been deprecated.") FCachedFocalLength
 {
 public:
 	float OriginalFocalLength = 0.0f;
@@ -49,22 +48,21 @@ public:
 	UFUNCTION(BlueprintCallable, Category = "Lens Distortion")
 	ULensFile* GetLensFile(const FLensFilePicker& Picker) const;
 
-	/** Return all handlers associated with the input camera component */
+	UE_DEPRECATED(5.1, "This function has been deprecated. The subsystem no longer tracks distortion handlers. Query for a handler from a specific Lens Component belonging to the input camera.")
 	UFUNCTION(BlueprintCallable, Category = "Lens Distortion")
 	TArray<ULensDistortionModelHandlerBase*> GetDistortionModelHandlers(UCineCameraComponent* Component);
 
-	/** 
-	 * Return the handler associated with the input distortion source, if one exists 
-	 * If bUpdatePicker is true, the input picker reference will be updated so that its properties match those of the found handler
-	 */
+PRAGMA_DISABLE_DEPRECATION_WARNINGS
+	UE_DEPRECATED(5.1, "This function has been deprecated. The subsystem no longer tracks distortion handlers. Query for a handler from a specific Lens Component belonging to the input camera.")
 	UFUNCTION(BlueprintCallable, Category = "Lens Distortion")
 	ULensDistortionModelHandlerBase* FindDistortionModelHandler(UPARAM(ref)FDistortionHandlerPicker& DistortionHandlerPicker, bool bUpdatePicker = true) const;
 
-	/** Return the handler associated with the input distortion source, if one exists that also matches the input model. If none exist, create a new handler and return it. */
+	UE_DEPRECATED(5.1, "This function has been deprecated. The subsystem no longer tracks distortion handlers. Query for a handler from a specific Lens Component belonging to the input camera.")
 	UFUNCTION(BlueprintCallable, Category = "Lens Distortion")
 	ULensDistortionModelHandlerBase* FindOrCreateDistortionModelHandler(UPARAM(ref)FDistortionHandlerPicker& DistortionHandlerPicker, const TSubclassOf<ULensModel> LensModelClass);
+PRAGMA_ENABLE_DEPRECATION_WARNINGS
 
-	/** Disassociate the input handler from the input camera component in the subsystem's handler map */
+	UE_DEPRECATED(5.1, "This function has been deprecated. The subsystem no longer tracks distortion handlers. Query for a handler from a specific Lens Component belonging to the input camera.")
 	UFUNCTION(BlueprintCallable, Category = "Lens Distortion")
 	void UnregisterDistortionModelHandler(UCineCameraComponent* Component, ULensDistortionModelHandlerBase* Handler);
 
@@ -120,10 +118,10 @@ public:
 	/** Remove a Lens Model from the registered model map */
 	void UnregisterDistortionModel(TSubclassOf<ULensModel> LensModel);
 
-	/** Update the original focal length of the input camera componet */
+	UE_DEPRECATED(5.1, "This function has been deprecated. The subsystem no longer tracks original focal length for distortion.")
 	void UpdateOriginalFocalLength(UCineCameraComponent* Component, float InFocalLength);
 
-	/** Update the overscanned focal length of the input camera component */
+	UE_DEPRECATED(5.1, "This function has been deprecated. The subsystem no longer tracks overscanned focal length for distortion.")
 	void UpdateOverscanFocalLength(UCineCameraComponent* Component, float InFocalLength);
 
 	/** Register a new overlay material name and path that can be queried from camera calibration tools */
@@ -132,10 +130,7 @@ public:
 	/** Unregister an overlay material */
 	void UnregisterOverlayMaterial(const FName& MaterialName);
 
-	/** 
-	 * Get the original focal length of the input camera component, if it exists in the subsystems map. 
-	 * Returns false and does not update the output parameter if the input camera is not in the map. 
-	 */
+	UE_DEPRECATED(5.1, "This function has been deprecated. The subsystem no longer tracks original focal length for distortion.")
 	bool GetOriginalFocalLength(UCineCameraComponent* Component, float& OutFocalLength);
 
 private:
@@ -165,12 +160,6 @@ private:
 	/** Map of actor components to the authoritative lens model that should be used with that component */
 	TMap<FObjectKey, TSubclassOf<ULensModel>> ComponentsWithAuthoritativeModels;
 
-	/** Map of actor components to the authoritative lens model that should be used with that component */
-	TMultiMap<FObjectKey, ULensDistortionModelHandlerBase*> LensDistortionHandlerMap;
-
-	/** Map of camera components to a cached pair of focal lengths for that camera */
-	TMap<FObjectKey, FCachedFocalLength> CachedFocalLengthMap;
-
 private:
 
 	FDelegateHandle PostEngineInitHandle;
diff --git a/Engine/Plugins/VirtualProduction/CameraCalibrationCore/Source/CameraCalibrationCore/Public/CameraCalibrationTypes.h b/Engine/Plugins/VirtualProduction/CameraCalibrationCore/Source/CameraCalibrationCore/Public/CameraCalibrationTypes.h
--- a/Engine/Plugins/VirtualProduction/CameraCalibrationCore/Source/CameraCalibrationCore/Public/CameraCalibrationTypes.h
+++ b/Engine/Plugins/VirtualProduction/CameraCalibrationCore/Source/CameraCalibrationCore/Public/CameraCalibrationTypes.h
@@ -9,11 +9,19 @@
 #include "CameraCalibrationTypes.generated.h"
 
 /** Utility structure for selecting a distortion handler from the camera calibration subsystem */
+struct UE_DEPRECATED(5.1, "This struct has been deprecated.") FDistortionHandlerPicker;
+
 USTRUCT(BlueprintType)
 struct CAMERACALIBRATIONCORE_API FDistortionHandlerPicker
 {
 	GENERATED_BODY()
 
+public:
+	/** Default destructor */
+PRAGMA_DISABLE_DEPRECATION_WARNINGS
+	~FDistortionHandlerPicker() = default;
+PRAGMA_ENABLE_DEPRECATION_WARNINGS
+
 public:
 	/** CineCameraComponent with which the desired distortion handler is associated */
 	UPROPERTY(BlueprintReadWrite, Category = "Distortion", Transient)
diff --git a/Engine/Plugins/VirtualProduction/CameraCalibrationCore/Source/CameraCalibrationCore/Public/LensComponent.h b/Engine/Plugins/VirtualProduction/CameraCalibrationCore/Source/CameraCalibrationCore/Public/LensComponent.h
--- a/Engine/Plugins/VirtualProduction/CameraCalibrationCore/Source/CameraCalibrationCore/Public/LensComponent.h
+++ b/Engine/Plugins/VirtualProduction/CameraCalibrationCore/Source/CameraCalibrationCore/Public/LensComponent.h
@@ -26,6 +26,52 @@ enum class EFIZEvaluationMode : uint8
 	DoNotEvaluate,
 };
 
+/** Controls whether this component can override the camera's filmback, and if so, which override to use */
+UENUM(BlueprintType)
+enum class EFilmbackOverrideSource : uint8
+{
+	/** Override the camera's filmback using the sensor dimensions recorded in the LensInfo of the LensFile */
+	LensFile,
+	/** Override the camera's filmback using the CroppedFilmback setting below */
+	CroppedFilmbackSetting,
+	/** Do not override the camera's filmback */
+	DoNotOverride,
+};
+
+/** Specifies from where the distortion state information comes */
+UENUM(BlueprintType)
+enum class EDistortionSource : uint8
+{
+	/** Distortion state is evaluated using the LensFile */
+	LensFile,
+	/** Distortion state is inputted directly from a LiveLink subject */
+	LiveLinkLensSubject,
+	/** Distortion state is set manually by the user using the Distortion State setting below */
+	Manual,
+};
+
+USTRUCT(BlueprintType)
+struct FLensFileEvaluationInputs
+{
+	GENERATED_BODY()
+
+public:
+	UPROPERTY(Interp, VisibleAnywhere, BlueprintReadOnly, Category = "Lens File")
+	float Focus = 0.0f;
+
+	UPROPERTY(Interp, VisibleAnywhere, BlueprintReadOnly, Category = "Lens File")
+	float Iris = 0.0f;
+
+	UPROPERTY(Interp, VisibleAnywhere, BlueprintReadOnly, Category = "Lens File")
+	float Zoom = 0.0f;
+
+	UPROPERTY(Interp, VisibleAnywhere, BlueprintReadOnly, Category = "Lens File")
+	FCameraFilmbackSettings Filmback;
+
+	UPROPERTY(Transient, BlueprintReadOnly, Category = "Lens File")
+	bool bIsValid = false;
+};
+
 /** Component for applying a post-process lens distortion effect to a CineCameraComponent on the same actor */
 UCLASS(HideCategories=(Tags, Activation, Cooking, AssetUserData, Collision), meta=(BlueprintSpawnableComponent))
 class CAMERACALIBRATIONCORE_API ULensComponent : public UActorComponent
@@ -52,22 +98,54 @@ public:
 #endif //WITH_EDITOR
 	//~ End UObject interface
 
+	/** Get the LensFile picker used by this component */
+	UFUNCTION(BlueprintPure, Category = "Lens Component")
+	FLensFilePicker GetLensFilePicker() const;
+
 	/** Get the LensFile used by this component */
 	UFUNCTION(BlueprintPure, Category="Lens Component")
 	ULensFile* GetLensFile() const;
 
+	/** Set the LensFile picker used by this component */
+	UFUNCTION(BlueprintCallable, Category = "Lens Component")
+	void SetLensFilePicker(FLensFilePicker LensFile);
+
 	/** Set the LensFile used by this component */
 	UFUNCTION(BlueprintCallable, Category="Lens Component")
-	void SetLensFile(FLensFilePicker LensFile);
+	void SetLensFile(ULensFile* LensFile);
 
 	/** Get the evaluation mode used to evaluate the LensFile */
-	UFUNCTION(BlueprintPure, Category="Lens Component")
+	UFUNCTION(BlueprintPure, Category="Lens Component", meta = (DisplayName = "Get FIZ Evaluation Mode"))
 	EFIZEvaluationMode GetFIZEvaluationMode() const;
 
 	/** Set the evaluation mode used to evaluate the LensFile */
 	UFUNCTION(BlueprintCallable, Category="Lens Component", meta=(DisplayName="Set FIZ Evaluation Mode"))
 	void SetFIZEvaluationMode(EFIZEvaluationMode Mode);
 
+	/** Get the evaluation mode used to evaluate the LensFile */
+	UFUNCTION(BlueprintPure, Category = "Lens Component")
+	float GetOverscanMultiplier() const;
+
+	/** Set the LensFile used by this component */
+	UFUNCTION(BlueprintCallable, Category = "Lens Component")
+	void SetOverscanMultiplier(float Multiplier);
+
+	/** Get the filmback override setting */
+	UFUNCTION(BlueprintPure, Category = "Lens Component")
+	EFilmbackOverrideSource GetFilmbackOverrideSetting() const;
+
+	/** Set the filmback override setting */
+	UFUNCTION(BlueprintCallable, Category = "Lens Component")
+	void SetFilmbackOverrideSetting(EFilmbackOverrideSource Setting);
+
+	/** Get the cropped filmback (only relevant if the filmback override setting is set to use the CroppedFilmback */
+	UFUNCTION(BlueprintPure, Category = "Lens Component")
+	FCameraFilmbackSettings GetCroppedFilmback() const;
+
+	/** Set the cropped filmback (only relevant if the filmback override setting is set to use the CroppedFilmback */
+	UFUNCTION(BlueprintCallable, Category = "Lens Component")
+	void SetCroppedFilmback(FCameraFilmbackSettings Filmback);
+
 	/** Returns true if nodal offset will be automatically applied during this component's tick, false otherwise */
 	UFUNCTION(BlueprintPure, Category="Lens Component")
 	bool ShouldApplyNodalOffsetOnTick() const;
@@ -76,10 +154,54 @@ public:
 	UFUNCTION(BlueprintCallable, Category="Lens Component")
 	void SetApplyNodalOffsetOnTick(bool bApplyNodalOffset);
 
+	/** Get the distortion source setting */
+	UFUNCTION(BlueprintPure, Category = "Lens Component")
+	EDistortionSource GetDistortionSource() const;
+
+	/** Set the distortion source setting */
+	UFUNCTION(BlueprintCallable, Category = "Lens Component")
+	void SetDistortionSource(EDistortionSource Source);
+
+	/** Whether distortion should be applied to the target camera */
+	UFUNCTION(BlueprintPure, Category = "Lens Component")
+	bool ShouldApplyDistortion() const;
+
+	/** Set whether distortion should be applied to the target camera */
+	UFUNCTION(BlueprintCallable, Category = "Lens Component")
+	void SetApplyDistortion(bool bApply);
+
+	/** Get the current lens model */
+	UFUNCTION(BlueprintPure, Category = "Lens Component")
+	TSubclassOf<ULensModel> GetLensModel() const;
+
+	/** Set the current lens model */
+	UFUNCTION(BlueprintCallable, Category = "Lens Component")
+	void SetLensModel(TSubclassOf<ULensModel> Model);
+
+	/** Get the current distortion state */
+	UFUNCTION(BlueprintPure, Category = "Lens Component")
+	FLensDistortionState GetDistortionState() const;
+
+	/** Set the current distortion state */
+	UFUNCTION(BlueprintCallable, Category = "Lens Component")
+	void SetDistortionState(FLensDistortionState State);
+
+	/** Get the original (not adjusted for overscan) focal length of the camera */
+	UFUNCTION(BlueprintPure, Category = "Lens Component")
+	float GetOriginalFocalLength() const;
+
+	/** Get the data used by this component to evaluate the LensFile */
+	UFUNCTION(BlueprintPure, Category = "Lens Component")
+	const FLensFileEvaluationInputs& GetLensFileEvaluationInputs() const;
+
 	/** Returns true if nodal offset was applied during the current tick, false otherwise */
 	UFUNCTION(BlueprintPure, Category = "Lens Component")
 	bool WasNodalOffsetAppliedThisTick() const;
 
+	/** Returns true if distortion was evaluated this tick */
+	UFUNCTION(BlueprintPure, Category = "Lens Component")
+	bool WasDistortionEvaluated() const;
+
 	/** 
 	 * Manually apply nodal offset to the specified component. 
 	 * If bUseManualInputs is true, the input Focus and Zoom values will be used to evaluate the LensFile .
@@ -88,6 +210,18 @@ public:
 	UFUNCTION(BlueprintCallable, Category="Lens Component", meta=(AdvancedDisplay=1))
 	void ApplyNodalOffset(USceneComponent* ComponentToOffset, bool bUseManualInputs = false, float ManualFocusInput = 0.0f, float ManualZoomInput = 0.0f);
 
+public:
+	/** Returns the LensDistortionHandler in use for the current LensModel */
+	ULensDistortionModelHandlerBase* GetLensDistortionHandler() const;
+
+	/** Reset the tracked component back to its original tracked pose and reapply nodal offset to it by re-evaluating the LensFile */
+	void ReapplyNodalOffset();
+
+PRAGMA_DISABLE_DEPRECATION_WARNINGS
+	UE_DEPRECATED(5.1, "This property is deprecated. Use GetDistortionSource() and SetDistortionSource() instead.")
+	FDistortionHandlerPicker GetDistortionHandlerPicker() const;
+PRAGMA_ENABLE_DEPRECATION_WARNINGS
+
 protected:
 	UE_DEPRECATED(5.1, "The use of this callback by this class has been deprecated and it is no longer registered. You can register your own delegate with FWorldDelegates::OnWorldPostActorTick")
 	void OnPostActorTick(UWorld* World, ELevelTick TickType, float DeltaSeconds);
@@ -96,6 +230,9 @@ private:
 	/** Evaluate the LensFile for nodal offset (using the current evaluation mode) and apply it to the latest component to offset */
 	void ApplyNodalOffset();
 
+	/** Evaluate the focal length from the LensFile and applies the calibrated value to the camera */
+	void EvaluateFocalLength(UCineCameraComponent* CineCameraComponent);
+
 	/** If TargetCameraComponent is not set, initialize it to the first CineCameraComponent on the same actor as this component */
 	void InitDefaultCamera();
 
@@ -103,7 +240,7 @@ private:
 	void CleanupDistortion(UCineCameraComponent* CineCameraComponent);
 
 	/** Register a new lens distortion handler with the camera calibration subsystem using the selected lens file */
-	void CreateDistortionHandlerForLensFile();
+	void CreateDistortionHandler();
 
 	/** Register to the new LiveLink component's callback to be notified when its controller map changes */
 	void OnLiveLinkComponentRegistered(ULiveLinkComponentController* LiveLinkComponent);
@@ -118,7 +255,10 @@ private:
 	void UpdateLiveLinkFIZ(const ULiveLinkComponentController* const LiveLinkComponent, const FLiveLinkSubjectFrameData& SubjectData);
 
 	/** Updates the focus and zoom inputs that will be used to evaluate the LensFile based on the evaluation mode */
-	void UpdateLensFileEvaluationInputs();
+	void UpdateLensFileEvaluationInputs(UCineCameraComponent* CineCameraComponent);
+
+	/** Updates the camera's filmback based on the filmback override settings */
+	void UpdateCameraFilmback(UCineCameraComponent* CineCameraComponent);
 
 protected:
 	/** Lens File used to drive distortion with current camera settings */
@@ -133,36 +273,53 @@ protected:
 	UPROPERTY(EditInstanceOnly, AdvancedDisplay, Category="Lens File", meta=(UseComponentPicker, AllowedClasses="/Script/CinematicCamera.CineCameraComponent"))
 	FComponentReference TargetCameraComponent;
 
-	/** 
-	 * If checked, nodal offset will be applied automatically when this component ticks. 
-	 * Set to false if nodal offset needs to be manually applied at some other time (via Blueprints).
-	 */
-	UPROPERTY(EditAnywhere, Category="Nodal Offset")
-	bool bApplyNodalOffsetOnTick = true;
+	/** Inputs to LensFile evaluation */
+	UPROPERTY(VisibleAnywhere, AdvancedDisplay, Category = "Lens File")
+	FLensFileEvaluationInputs EvalInputs;
 
-	/** Structure used to query the camera calibration subsystem for a lens distortion model handler */
-	UPROPERTY(EditAnywhere, BlueprintReadWrite, Category="Distortion")
-	FDistortionHandlerPicker DistortionSource;
+	/** Specifies from where the distortion state information comes */
+	UPROPERTY(EditAnywhere, Category = "Distortion", meta = (DisplayName = "Distortion Source"))
+	EDistortionSource DistortionStateSource = EDistortionSource::LensFile;
 
 	/** Whether or not to apply distortion to the target camera component */
-	UPROPERTY(EditAnywhere, BlueprintReadWrite, Category="Distortion")
+	UPROPERTY(EditAnywhere, BlueprintReadWrite, Getter = ShouldApplyDistortion, Setter = SetApplyDistortion, Category = "Distortion")
 	bool bApplyDistortion = false;
 
-	/**
-	 * Whether to use the specified lens file to drive distortion
-	 * Enabling this will create a new distortion source and automatically set this component to use it
-	 */
-	UPROPERTY(EditAnywhere, AdvancedDisplay, Category="Distortion")
-	bool bEvaluateLensFileForDistortion = false;
+	/** The current lens model used for distortion */
+	UPROPERTY(EditAnywhere, Category = "Distortion", meta = (EditCondition = "DistortionStateSource == EDistortionSource::Manual"))
+	TSubclassOf<ULensModel> LensModel;
+
+	/** The current distortion state */
+	UPROPERTY(Interp, EditAnywhere, Category = "Distortion", meta = (ShowOnlyInnerProperties))
+	FLensDistortionState DistortionState;
 
 	/** Whether to scale the computed overscan by the overscan percentage */
-	UPROPERTY(AdvancedDisplay, BlueprintReadWrite, Category="Distortion")
+	UPROPERTY(AdvancedDisplay, BlueprintReadWrite, Category = "Distortion")
 	bool bScaleOverscan = false;
 
 	/** The percentage of the computed overscan that should be applied to the target camera */
-	UPROPERTY(EditAnywhere, AdvancedDisplay, BlueprintReadWrite, Category="Distortion", meta=(EditCondition="bScaleOverscan", ClampMin="0.0", ClampMax="2.0"))
+	UPROPERTY(EditAnywhere, AdvancedDisplay, BlueprintReadWrite, Category = "Distortion", meta = (EditCondition = "bScaleOverscan", ClampMin = "0.0", ClampMax = "2.0"))
 	float OverscanMultiplier = 1.0f;
 
+	/** Controls whether this component can override the camera's filmback, and if so, which override to use */
+	UPROPERTY(EditAnywhere, Category = "Filmback")
+	EFilmbackOverrideSource FilmbackOverride = EFilmbackOverrideSource::DoNotOverride;
+
+	/** Cropped filmback to use if the filmback override settings are set to use it */
+	UPROPERTY(EditAnywhere, Category = "Filmback", meta=(EditCondition="FilmbackOverride == EFilmbackOverrideSource::CroppedFilmbackSetting"))
+	FCameraFilmbackSettings CroppedFilmback;
+
+	/** 
+	 * If checked, nodal offset will be applied automatically when this component ticks. 
+	 * Set to false if nodal offset needs to be manually applied at some other time (via Blueprints).
+	 */
+	UPROPERTY(EditAnywhere, Category="Nodal Offset")
+	bool bApplyNodalOffsetOnTick = true;
+
+	/** Serialized transform of the TrackedComponent prior to nodal offset being applied */
+	UPROPERTY(Interp, VisibleAnywhere, AdvancedDisplay, Category = "Nodal Offset")
+	FTransform OriginalTrackedComponentTransform;
+
 	/** Whether a distortion effect is currently being applied to the target camera component */
 	UPROPERTY()
 	bool bIsDistortionSetup = false;
@@ -171,6 +328,26 @@ protected:
 	UPROPERTY()
 	float OriginalFocalLength = 35.0f;
 
+	/** Cached MID last applied to the target camera */
+	UPROPERTY()
+	UMaterialInstanceDynamic* LastDistortionMID = nullptr;
+
+	/** Cached most recent target camera, used to clean up the old camera when the user changes the target */
+	UPROPERTY()
+	UCineCameraComponent* LastCameraComponent = nullptr;
+
+	/** Map of lens models to handlers */
+	UPROPERTY(Transient)
+	TMap<TSubclassOf<ULensModel>, TObjectPtr<ULensDistortionModelHandlerBase>> LensDistortionHandlerMap;
+
+	/** Scene component that should have nodal offset applied */
+	UPROPERTY(Transient)
+	TWeakObjectPtr<USceneComponent> TrackedComponent;
+
+	/** Serialized name of the TrackedComponent, used to determine which component to re-apply nodal offset to in spawnables */
+	UPROPERTY()
+	FString TrackedComponentName;
+
 #if WITH_EDITORONLY_DATA
 	UE_DEPRECATED(5.1, "This property has been deprecated. The LensDistortion component no longer tracks the attached camera's original rotation.")
 	UPROPERTY()
@@ -179,38 +356,39 @@ protected:
 	UE_DEPRECATED(5.1, "This property has been deprecated. The LensDistortion component no longer tracks the attached camera's original location.")
 	UPROPERTY()
 	FVector OriginalCameraLocation_DEPRECATED;
-#endif //WITH_EDITORONLY_DATA
 
-	/** Cached MID last applied to the target camera */
+	UE_DEPRECATED(5.1, "This property has been deprecated. Use the DistortionStateSource to specify whether the lens file should be evaluated for distortion.")
 	UPROPERTY()
-	UMaterialInstanceDynamic* LastDistortionMID = nullptr;
+	bool bEvaluateLensFileForDistortion_DEPRECATED = false;
 
-	/** Cached most recent target camera, used to clean up the old camera when the user changes the target */
-	UPROPERTY(Transient)
-	UCineCameraComponent* LastCameraComponent = nullptr;
-
-	/** Distortion handler produced by this component */
+	UE_DEPRECATED(5.1, "This property has been deprecated. Use the LensDistortionHandlerMap to get a handler for the current LensModel.")
 	UPROPERTY(Transient)
-	ULensDistortionModelHandlerBase* ProducedLensDistortionHandler = nullptr;
+	ULensDistortionModelHandlerBase* LensDistortionHandler_DEPRECATED = nullptr;
 
-	/** Unique identifier representing the source of distortion data */
+	UE_DEPRECATED(5.1, "This property has been deprecated. Producer GUIDs are no longer used to identify distortion handlers. Use the LensDistortionHandlerMap to get a handler for the current LensModel.")
 	UPROPERTY(DuplicateTransient)
-	FGuid DistortionProducerID;
+	FGuid DistortionProducerID_DEPRECATED;
 
-	/** Inputs to LensFile evaluation */
-	UPROPERTY(Interp, VisibleAnywhere, AdvancedDisplay, Category="Lens File")
-	float EvalFocus = 0.0f;
-	
-	UPROPERTY(Interp, VisibleAnywhere, AdvancedDisplay, Category="Lens File")
-	float EvalZoom = 0.0f;
+PRAGMA_DISABLE_DEPRECATION_WARNINGS
+	UE_DEPRECATED(5.1, "This property has been deprecated. The handler picker is no longer used to identify a distortion handler. Use the LensDistortionHandlerMap to get a handler for the current LensModel.")
+	UPROPERTY(meta=(DeprecatedProperty, DeprecationMessage="This property has been deprecated. Use GetDistortionSource() and SetDistortionSource() instead."))
+	FDistortionHandlerPicker DistortionSource_DEPRECATED;
+PRAGMA_ENABLE_DEPRECATION_WARNINGS
+#endif //WITH_EDITORONLY_DATA
 
 private:
-	/** Scene component that should have nodal offset applied */
-	TWeakObjectPtr<USceneComponent> TrackedComponent;
-
 	/** Latest LiveLink FIZ data, used to evaluate the LensFile */
 	FLensFileEvalData LiveLinkFIZ;
 
+	/** Whether LiveLink FIZ was received this tick */
+	bool bWasLiveLinkFIZUpdated = false;
+
+	/** Whether distortion was evaluated this tick */
+	bool bWasDistortionEvaluated = false;
+
 	/** Whether or not nodal offset was applied to a tracked component this tick */
 	bool bWasNodalOffsetAppliedThisTick = false;
+
+	/** Latest focal length of the target camera, used to track external changes to the original (no overscan applied) focal length */
+	float LastFocalLength = -1.0f;
 };
diff --git a/Engine/Plugins/VirtualProduction/CameraCalibrationCore/Source/CameraCalibrationCore/Public/LensFile.h b/Engine/Plugins/VirtualProduction/CameraCalibrationCore/Source/CameraCalibrationCore/Public/LensFile.h
--- a/Engine/Plugins/VirtualProduction/CameraCalibrationCore/Source/CameraCalibrationCore/Public/LensFile.h
+++ b/Engine/Plugins/VirtualProduction/CameraCalibrationCore/Source/CameraCalibrationCore/Public/LensFile.h
@@ -5,6 +5,7 @@
 #include "Tickable.h"
 #include "UObject/Object.h"
 
+#include "CineCameraComponent.h"
 #include "CoreTypes.h"
 #include "Engine/Texture.h"
 #include "ICalibratedMapProcessor.h"
@@ -118,7 +119,7 @@ public:
 	/** Update the resolution used for intermediate blending displacement maps and for STMap derived data */
 	void UpdateDisplacementMapResolution(const FIntPoint NewDisplacementMapResolution);
 
-	/** Whether the sensor dimensions in the lens file will be compatible with the sensor dimensions of the input CineCameraComponent */
+	UE_DEPRECATED(5.1, "This function is deprecated. The LensFile checks this internally when it evaluates for distortion.")
 	bool IsCineCameraCompatible(const UCineCameraComponent* CineCameraComponent) const;
 
 	/** Update the input tolerance used when adding points to calibration tables */
@@ -386,5 +387,3 @@ struct FLensFileEvalData
 		NodalOffset.bWasApplied = false;
 	}
 };
-
-
diff --git a/Engine/Plugins/VirtualProduction/CameraCalibrationCore/Source/CameraCalibrationCoreEditor/Private/CameraCalibrationCoreEditorModule.cpp b/Engine/Plugins/VirtualProduction/CameraCalibrationCore/Source/CameraCalibrationCoreEditor/Private/CameraCalibrationCoreEditorModule.cpp
--- a/Engine/Plugins/VirtualProduction/CameraCalibrationCore/Source/CameraCalibrationCoreEditor/Private/CameraCalibrationCoreEditorModule.cpp
+++ b/Engine/Plugins/VirtualProduction/CameraCalibrationCore/Source/CameraCalibrationCoreEditor/Private/CameraCalibrationCoreEditorModule.cpp
@@ -9,7 +9,6 @@
 #include "CalibrationPointComponent.h"
 #include "CalibrationPointComponentDetails.h"
 #include "CameraCalibrationCoreEditorStyle.h"
-#include "DistortionHandlerPickerDetailCustomization.h"
 #include "Editor.h"
 #include "IPlacementModeModule.h"
 #include "LensFile.h"
@@ -28,11 +27,6 @@ void FCameraCalibrationCoreEditorModule::StartupModule()
 
 	FPropertyEditorModule& PropertyEditorModule = FModuleManager::Get().LoadModuleChecked<FPropertyEditorModule>("PropertyEditor");
 
-	PropertyEditorModule.RegisterCustomPropertyTypeLayout(
-		FDistortionHandlerPicker::StaticStruct()->GetFName(), 
-		FOnGetPropertyTypeCustomizationInstance::CreateStatic(&FDistortionHandlerPickerDetailCustomization::MakeInstance)
-	);
-
 	PropertyEditorModule.RegisterCustomClassLayout(
 		UCalibrationPointComponent::StaticClass()->GetFName(),
 		FOnGetDetailCustomizationInstance::CreateStatic(&FCalibrationPointComponentDetails::MakeInstance)
@@ -50,8 +44,6 @@ void FCameraCalibrationCoreEditorModule::ShutdownModule()
 		PropertyModule.UnregisterCustomClassLayout(ULensFile::StaticClass()->GetFName());
 		PropertyModule.UnregisterCustomClassLayout(UCalibrationPointComponent::StaticClass()->GetFName());
 
-		PropertyModule.UnregisterCustomPropertyTypeLayout(FDistortionHandlerPicker::StaticStruct()->GetFName());
-
 		UnregisterPlacementModeItems();
 	}
 }
diff --git a/Engine/Plugins/VirtualProduction/CameraCalibrationCore/Source/CameraCalibrationCoreEditor/Private/DistortionHandlerPickerDetailCustomization.cpp b/Engine/Plugins/VirtualProduction/CameraCalibrationCore/Source/CameraCalibrationCoreEditor/Private/DistortionHandlerPickerDetailCustomization.cpp
deleted file mode 100644
diff --git a/Engine/Plugins/VirtualProduction/CameraCalibrationCore/Source/CameraCalibrationCoreEditor/Private/DistortionHandlerPickerDetailCustomization.h b/Engine/Plugins/VirtualProduction/CameraCalibrationCore/Source/CameraCalibrationCoreEditor/Private/DistortionHandlerPickerDetailCustomization.h
deleted file mode 100644
diff --git a/Engine/Plugins/VirtualProduction/LiveLinkCamera/Source/LiveLinkCamera/Private/LiveLinkCameraController.cpp b/Engine/Plugins/VirtualProduction/LiveLinkCamera/Source/LiveLinkCamera/Private/LiveLinkCameraController.cpp
--- a/Engine/Plugins/VirtualProduction/LiveLinkCamera/Source/LiveLinkCamera/Private/LiveLinkCameraController.cpp
+++ b/Engine/Plugins/VirtualProduction/LiveLinkCamera/Source/LiveLinkCamera/Private/LiveLinkCameraController.cpp
@@ -17,57 +17,8 @@
 
 DEFINE_LOG_CATEGORY_STATIC(LogLiveLinkCameraController, Log, All);
 
-namespace LiveLinkCameraControllerUtils
-{
-	/** 
-	 * Finds indices of neighbor focus points for a given focus value and verify if both have only one zoom point 
-	 * Copied from private LensTablesUtils.
-	 */
-	template<typename Type>
-	bool HasSingleZoomPointInvolved(float InFocus, TConstArrayView<Type> Container)
-	{
-		if (Container.Num() <= 0)
-		{
-			return false;
-		}
-
-		int32 FirstFocusPointIndex = INDEX_NONE;
-		int32 SecondFocusPointIndex = INDEX_NONE;
-		for (int32 Index = 0; Index < Container.Num(); ++Index)
-		{
-			const Type& Point = Container[Index];
-			if (Point.Focus > InFocus)
-			{
-				SecondFocusPointIndex = Index;
-				FirstFocusPointIndex = FMath::Max(Index - 1, 0);
-				break;
-			}
-			else if (FMath::IsNearlyEqual(Point.Focus, InFocus))
-			{
-				//We found a point exactly matching the desired one
-				SecondFocusPointIndex = Index;
-				FirstFocusPointIndex = Index;
-				break;
-			}
-		}
-
-		//We haven't found a point, default to last one
-		if (FirstFocusPointIndex == INDEX_NONE || SecondFocusPointIndex == INDEX_NONE)
-		{
-			SecondFocusPointIndex = Container.Num() - 1;
-			FirstFocusPointIndex = Container.Num() - 1;
-		}
-
-		return Container[FirstFocusPointIndex].GetNumPoints() <= 1 && Container[SecondFocusPointIndex].GetNumPoints() <= 1;
-	}
-}
-
 ULiveLinkCameraController::ULiveLinkCameraController()
 {
-	if (!HasAnyFlags(RF_ArchetypeObject | RF_ClassDefaultObject))
-	{
-		DistortionProducerID = FGuid::NewGuid();
-	}
 }
 
 void ULiveLinkCameraController::Tick(float DeltaTime, const FLiveLinkSubjectFrameData& SubjectData)
@@ -75,14 +26,6 @@ void ULiveLinkCameraController::Tick(float DeltaTime, const FLiveLinkSubjectFram
 	// Invalidate the lens file evaluation data
 	LensFileEvalData.Invalidate();
 
-	if (UCameraComponent* CameraComponent = Cast<UCameraComponent>(GetAttachedComponent()))
-	{
-		if (AActor* Camera = Cast<AActor>(GetOuterActor()))
-		{
-			LensFileEvalData.Camera.UniqueId = Camera->GetUniqueID();
-		}
-	}
-
 	const FLiveLinkCameraStaticData* StaticData = SubjectData.StaticData.Cast<FLiveLinkCameraStaticData>();
 	const FLiveLinkCameraFrameData* FrameData = SubjectData.FrameData.Cast<FLiveLinkCameraFrameData>();
 
@@ -133,81 +76,25 @@ void ULiveLinkCameraController::Tick(float DeltaTime, const FLiveLinkSubjectFram
 				ULensFile* SelectedLensFile = LensFilePicker.GetLensFile();
 				LensFileEvalData.LensFile = SelectedLensFile;
 
-				// If we're using a lens file, use its sensor dimensions unless Live Link is streaming different values
-				if (SelectedLensFile)
-				{
-					if (bUseCroppedFilmback)
-					{
-						LensFileEvalData.Distortion.Filmback.X = CroppedFilmback.SensorWidth;
-					}
-					else if ((StaticData->FilmBackWidth > 0.0f) && UpdateFlags.bApplyFilmBack) 
-					{ 
-						LensFileEvalData.Distortion.Filmback.X = StaticData->FilmBackWidth;
-					}
-					else
-					{
-						LensFileEvalData.Distortion.Filmback.X = SelectedLensFile->LensInfo.SensorDimensions.X;
-					}
-
-					if (bUseCroppedFilmback)
-					{
-						LensFileEvalData.Distortion.Filmback.Y = CroppedFilmback.SensorHeight;
-					}
-					else if ((StaticData->FilmBackHeight > 0.0f) && UpdateFlags.bApplyFilmBack)
-					{
-						LensFileEvalData.Distortion.Filmback.Y = StaticData->FilmBackHeight;
-					}
-					else
-					{
-						LensFileEvalData.Distortion.Filmback.Y = SelectedLensFile->LensInfo.SensorDimensions.Y;
-					}
-
-					CineCameraComponent->Filmback.SensorWidth = LensFileEvalData.Distortion.Filmback.X;
-					CineCameraComponent->Filmback.SensorHeight = LensFileEvalData.Distortion.Filmback.Y;
-				}
-				else
-				{
-					if (StaticData->FilmBackWidth > 0.0f && UpdateFlags.bApplyFilmBack) 
-					{ 
-						CineCameraComponent->Filmback.SensorWidth = StaticData->FilmBackWidth; 
-					}
-
-					if (StaticData->FilmBackHeight > 0.0f && UpdateFlags.bApplyFilmBack) 
-					{ 
-						CineCameraComponent->Filmback.SensorHeight = StaticData->FilmBackHeight; 
-					}
+				if (StaticData->FilmBackWidth > 0.0f && UpdateFlags.bApplyFilmBack) 
+				{ 
+					CineCameraComponent->Filmback.SensorWidth = StaticData->FilmBackWidth; 
 				}
 
-				// If the filmback changed and there is a lens file selected, warn the user if it is incompatible with it
-				if (LastFilmback != CineCameraComponent->Filmback)
-				{
-					if (SelectedLensFile && SelectedLensFile->IsCineCameraCompatible(CineCameraComponent) == false)
-					{
-						UE_LOG(LogLiveLinkCameraController, Warning, TEXT("LensFile '%s' has a smaller sensor size than the CameraComponent of '%s' (driven by LiveLinkCameraController '%s')")
-							, *SelectedLensFile->GetName()
-							, *CineCameraComponent->GetOwner()->GetName()
-							, *this->GetName());
-					}
+				if (StaticData->FilmBackHeight > 0.0f && UpdateFlags.bApplyFilmBack) 
+				{ 
+					CineCameraComponent->Filmback.SensorHeight = StaticData->FilmBackHeight; 
 				}
 
 				//Verify different Lens Tables based on streamed FIZ at some intervals
-				static const double TableVerificationInterval = 10;
-				if ((FPlatformTime::Seconds() - LastLensTableVerificationTimestamp) >= TableVerificationInterval)
-				{
-					LastLensTableVerificationTimestamp = FPlatformTime::Seconds();
-					VerifyFIZWithLensFileTables(SelectedLensFile, StaticData);
-				}
+ 				static const double TableVerificationInterval = 10;
+ 				if ((FPlatformTime::Seconds() - LastLensTableVerificationTimestamp) >= TableVerificationInterval)
+ 				{
+ 					LastLensTableVerificationTimestamp = FPlatformTime::Seconds();
+ 					VerifyFIZWithLensFileTables(SelectedLensFile, StaticData);
+ 				}
 
 				ApplyFIZ(SelectedLensFile, CineCameraComponent, StaticData, FrameData);
-				ApplyDistortion(SelectedLensFile, CineCameraComponent, StaticData, FrameData);
-
-				LastFilmback = CineCameraComponent->Filmback;
-
-				UCameraCalibrationSubsystem* const SubSystem = GEngine->GetEngineSubsystem<UCameraCalibrationSubsystem>();
-				if (SubSystem)
-				{
-					SubSystem->UpdateOriginalFocalLength(CineCameraComponent, CineCameraComponent->CurrentFocalLength);
-				}
 			}
 
 #if WITH_EDITORONLY_DATA
@@ -235,114 +122,11 @@ TSubclassOf<UActorComponent> ULiveLinkCameraController::GetDesiredComponentClass
 	return UCameraComponent::StaticClass();
 }
 
-void ULiveLinkCameraController::SetAttachedComponent(UActorComponent* ActorComponent)
-{
-	Super::SetAttachedComponent(ActorComponent);
-
-	if (UCineCameraComponent* const CineCameraComponent = Cast<UCineCameraComponent>(GetAttachedComponent()))
-	{	
-		// Initialize the most recent filmback and component transform to the current values of the camera to properly detect changes to this property
-		LastFilmback = CineCameraComponent->Filmback;
-
-		ULensFile* SelectedLensFile = LensFilePicker.GetLensFile();
-		if (SelectedLensFile && SelectedLensFile->IsCineCameraCompatible(CineCameraComponent) == false)
-		{
-			UE_LOG(LogLiveLinkCameraController, Warning, TEXT("LensFile '%s' has a smaller sensor size than the CameraComponent of '%s' (driven by LiveLinkCameraController '%s')")
-				, *SelectedLensFile->GetName()
-				, *CineCameraComponent->GetOwner()->GetName()
-				, *this->GetName());
-		}
-	}
-}
-
-void ULiveLinkCameraController::Cleanup()
-{
-	if (UCineCameraComponent* const CineCameraComponent = Cast<UCineCameraComponent>(GetAttachedComponent()))
-	{
-		if (UCameraCalibrationSubsystem* SubSystem = GEngine->GetEngineSubsystem<UCameraCalibrationSubsystem>())
-		{
-			SubSystem->UnregisterDistortionModelHandler(CineCameraComponent, LensDistortionHandler);
-		}
-	}
-}
-
-void ULiveLinkCameraController::PostDuplicate(bool bDuplicateForPIE)
-{
-	Super::PostDuplicate(bDuplicateForPIE);
-
-	if (!DistortionProducerID.IsValid())
-	{
-		DistortionProducerID = FGuid::NewGuid();
-	}
-}
-
-void ULiveLinkCameraController::PostEditImport()
-{
-	Super::PostEditImport();
-
-	if (!DistortionProducerID.IsValid())
-	{
-		DistortionProducerID = FGuid::NewGuid();
-	}
-}
-
 void ULiveLinkCameraController::SetApplyNodalOffset(bool bInApplyNodalOffset)
 {
 	// The use of this function has been deprecated.
 }
 
-#if WITH_EDITOR
-bool ULiveLinkCameraController::CanEditChange(const FProperty* InProperty) const
-{
-	const FName PropertyName = InProperty->GetFName();
-	if (PropertyName == GET_MEMBER_NAME_CHECKED(ULiveLinkCameraController, bUseCroppedFilmback))
-	{
-		if (LensFilePicker.GetLensFile())
-		{
-			return true;
-		}
-		return false;
-	}
-	else if (PropertyName == GET_MEMBER_NAME_CHECKED(ULiveLinkCameraController, CroppedFilmback))
-	{
-		if (LensFilePicker.GetLensFile() && bUseCroppedFilmback)
-		{
-			return true;
-		}
-		return false;
-	}
-
-	return Super::CanEditChange(InProperty);
-}
-
-void ULiveLinkCameraController::PostEditChangeProperty(struct FPropertyChangedEvent& PropertyChangedEvent)
-{
-	const FName PropertyName = (PropertyChangedEvent.Property != NULL) ? PropertyChangedEvent.Property->GetFName() : NAME_None;
-
-	if (PropertyName == GET_MEMBER_NAME_CHECKED(FLensFilePicker, LensFile)
-		|| PropertyName == GET_MEMBER_NAME_CHECKED(FLensFilePicker, bUseDefaultLensFile))
-	{
-		if (UCineCameraComponent* const CineCameraComponent = Cast<UCineCameraComponent>(GetAttachedComponent()))
-		{
-			ULensFile* SelectedLensFile = LensFilePicker.GetLensFile();
-			if (SelectedLensFile)
-			{
-				if (SelectedLensFile->IsCineCameraCompatible(CineCameraComponent) == false)
-				{
-					UE_LOG(LogLiveLinkCameraController, Warning, TEXT("LensFile '%s' has a smaller sensor size than the CameraComponent of '%s' (driven by LiveLinkCameraController '%s')")
-						, *SelectedLensFile->GetName()
-						, *CineCameraComponent->GetName()
-						, *this->GetName());
-				}
-			}
-		}
-
-		//When LensFile is changed, force update Table verification
-		LastLensTableVerificationTimestamp = 0;
-	}
-}
-#endif
-
 void ULiveLinkCameraController::ApplyFIZ(ULensFile* LensFile, UCineCameraComponent* CineCameraComponent, const FLiveLinkCameraStaticData* StaticData, const FLiveLinkCameraFrameData* FrameData)
 {
 	/**
@@ -388,46 +172,6 @@ void ULiveLinkCameraController::ApplyFIZ(ULensFile* LensFile, UCineCameraCompone
 			CineCameraComponent->LensSettings.MinFStop = FMath::Min(CineCameraComponent->LensSettings.MinFStop, CineCameraComponent->CurrentAperture);
 			CineCameraComponent->LensSettings.MaxFStop = FMath::Max(CineCameraComponent->LensSettings.MaxFStop, CineCameraComponent->CurrentAperture);
 		}
-		
-		if (UpdateFlags.bApplyFocalLength)
-		{
-			//To evaluate focal length, we need F/Z pair. If F or Z is not streamed, we use the default value (0)
-			
-			const float FocusValue = LensFileEvalData.Input.Focus;
-			bool bHasValidFocalLength = false;
-			const float ZoomValue = LensFileEvalData.Input.Zoom;
-			FFocalLengthInfo FocalLengthInfo;
-			if (LensFile->EvaluateFocalLength(FocusValue, ZoomValue, FocalLengthInfo))
-			{
-				if ((FocalLengthInfo.FxFy[0] > KINDA_SMALL_NUMBER) && (FocalLengthInfo.FxFy[1] > KINDA_SMALL_NUMBER))
-				{
-					// This is how field of view, filmback, and focal length are related:
-					//
-					// FOVx = 2*atan(1/(2*Fx)) = 2*atan(FilmbackX / (2*FocalLength))
-					// => FocalLength = Fx*FilmbackX
-					// 
-					// FOVy = 2*atan(1/(2*Fy)) = 2*atan(FilmbackY / (2*FocalLength))
-					// => FilmbackY = FocalLength / Fy
-
-					// Adjust FocalLength and Filmback to match FxFy (which has already been divided by resolution in pixels)
-					const float NewFocalLength = CineCameraComponent->CurrentFocalLength = FocalLengthInfo.FxFy[0] * CineCameraComponent->Filmback.SensorWidth;
-					CineCameraComponent->Filmback.SensorHeight = CineCameraComponent->CurrentFocalLength / FocalLengthInfo.FxFy[1];
-
-					// Update the minimum and maximum focal length of the camera (if needed)
-					CineCameraComponent->LensSettings.MinFocalLength = FMath::Max(FMath::Min(CineCameraComponent->LensSettings.MinFocalLength, NewFocalLength), 0.01);
-					CineCameraComponent->LensSettings.MaxFocalLength = FMath::Max(FMath::Max(CineCameraComponent->LensSettings.MaxFocalLength, NewFocalLength), 0.01);
-
-					CineCameraComponent->SetCurrentFocalLength(NewFocalLength);
-					bHasValidFocalLength = true;
-				}
-			}
-
-			//If FocalLength could not be applied and it's streamed in, use LiveLink input directly without affecting filmback
-			if (StaticData->bIsFocalLengthSupported && bHasValidFocalLength == false)
-			{
-				CineCameraComponent->SetCurrentFocalLength(ZoomValue);
-			}
-		}
 	}
 	else
 	{
@@ -445,28 +189,38 @@ void ULiveLinkCameraController::ApplyFIZ(ULensFile* LensFile, UCineCameraCompone
 				const float NewAperture = FMath::Lerp(CineCameraComponent->LensSettings.MinFStop, CineCameraComponent->LensSettings.MaxFStop, FrameData->Aperture);
 				CineCameraComponent->CurrentAperture = NewAperture;
 			}
-			if (StaticData->bIsFocalLengthSupported && UpdateFlags.bApplyFocalLength)
-			{
-				const float NewZoom = FMath::Lerp(CineCameraComponent->LensSettings.MinFocalLength, CineCameraComponent->LensSettings.MaxFocalLength, FrameData->FocalLength);
-				CineCameraComponent->SetCurrentFocalLength(NewZoom);
-			}
 		}
 		else
 		{
 			if (StaticData->bIsFocusDistanceSupported && UpdateFlags.bApplyFocusDistance)
 			{
-				CineCameraComponent->FocusSettings.ManualFocusDistance = FrameData->FocusDistance;
+				if (FrameData->FocusDistance > 1.0f)
+				{
+					CineCameraComponent->FocusSettings.ManualFocusDistance = FrameData->FocusDistance;
+				}
 			}
 
 			if (StaticData->bIsApertureSupported && UpdateFlags.bApplyAperture)
 			{
-				CineCameraComponent->CurrentAperture = FrameData->Aperture;
+				if (FrameData->Aperture > 1.0f)
+				{
+					CineCameraComponent->CurrentAperture = FrameData->Aperture;
+				}
 			}
+		}
+	}
 
-			if (StaticData->bIsFocalLengthSupported && UpdateFlags.bApplyFocalLength)
-			{
-				CineCameraComponent->SetCurrentFocalLength(FrameData->FocalLength);
-			}
+	// Handle Focal Length
+	if (StaticData->bIsFocalLengthSupported && UpdateFlags.bApplyFocalLength)
+	{
+		if (bUseCameraRange)
+		{
+			const float NewZoom = FMath::Lerp(CineCameraComponent->LensSettings.MinFocalLength, CineCameraComponent->LensSettings.MaxFocalLength, FrameData->FocalLength);
+			CineCameraComponent->SetCurrentFocalLength(NewZoom);
+		}
+		else if (FrameData->FocalLength > 1.0f)
+		{
+			CineCameraComponent->SetCurrentFocalLength(FrameData->FocalLength);
 		}
 	}
 }
@@ -478,35 +232,7 @@ void ULiveLinkCameraController::ApplyNodalOffset(ULensFile* SelectedLensFile, UC
 
 void ULiveLinkCameraController::ApplyDistortion(ULensFile* LensFile, UCineCameraComponent* CineCameraComponent, const FLiveLinkCameraStaticData* StaticData, const FLiveLinkCameraFrameData* FrameData)
 {
-	//Even if no FIZ was streamed in, we evaluate LensFile for distortion using default values (0.0f) like we do when applying FIZ
-	if (LensFile)
-	{
-		UCameraCalibrationSubsystem* const SubSystem = GEngine->GetEngineSubsystem<UCameraCalibrationSubsystem>();
-		const FString HandlerDisplayName = FString::Format(TEXT("{0} (Lens File)"), { LensFile->GetFName().ToString() });
-
-		FDistortionHandlerPicker DistortionHandlerPicker = { CineCameraComponent, DistortionProducerID, HandlerDisplayName };
-		LensDistortionHandler = SubSystem->FindOrCreateDistortionModelHandler(DistortionHandlerPicker, LensFile->LensInfo.LensModel);
-
-		if (LensDistortionHandler)
-		{
-			//Go through the lens file to get distortion data based on FIZ
-			//Our handler's displacement map will get updated
-			LensFileEvalData.Distortion.bWasEvaluated = true;
-
-			//Evaluate distortion using filmback not including SensorHeight modification based on Fy
-			LensFile->EvaluateDistortionData(LensFileEvalData.Input.Focus
-											, LensFileEvalData.Input.Zoom
-											, LensFileEvalData.Distortion.Filmback
-											, LensDistortionHandler);
-
-			// Adjust overscan by the overscan multiplier
-			if (bScaleOverscan)
-			{
-				const float ScaledOverscanFactor = ((LensDistortionHandler->GetOverscanFactor() - 1.0f) * OverscanMultiplier) + 1.0f;
-				LensDistortionHandler->SetOverscanFactor(ScaledOverscanFactor);
-			}
-		}
-	}
+	// The use of this function has been deprecated.
 }
 
 void ULiveLinkCameraController::OnPostActorTick(UWorld* World, ELevelTick TickType, float DeltaSeconds)
@@ -541,92 +267,6 @@ void ULiveLinkCameraController::VerifyFIZWithLensFileTables(ULensFile* LensFile,
 				FLiveLinkLog::WarningOnce(NAME_InvalidIrisMappingWhenNotStreamed, FakedSubjectKey, TEXT("Problem applying Iris for subject '%s' using LensFile '%s'. Iris wasn't streamed in and more than one iris mapping was found."), *FakedSubjectKey.SubjectName.ToString(), *LensFile->GetName());
 			}
 		}
-
-		if (UpdateFlags.bApplyFocalLength)
-		{
-			if (StaticData->bIsFocusDistanceSupported == false)
-			{
-				// If the LiveLink source is not streaming focus, we will default to evaluate at 0.0f.
-				// In this case, it is expected that there is only one mapping, so warn the user if this is not the case.
-				const int32 FocalLengthFocusPointCount = LensFile->FocalLengthTable.GetFocusPointNum();
-				if (FocalLengthFocusPointCount > 0 && FocalLengthFocusPointCount != 1)
-				{
-					static const FName NAME_InvalidFocalLengthMappingWhenNoFocusStreamed = "LiveLinkCamera_FocusNotStreamedInvalidFocalLengthMapping";
-					FLiveLinkLog::WarningOnce(NAME_InvalidFocalLengthMappingWhenNoFocusStreamed, FakedSubjectKey, TEXT("Problem applying FocalLength using subject '%s' and LensFile '%s'. Focus wasn't streamed in and more than one focal length focus point was found."), *FakedSubjectKey.SubjectName.ToString(), *LensFile->GetName());
-				}
-			}
-
-			if (StaticData->bIsFocalLengthSupported == false)
-			{
-				// If the LiveLink source is not streaming zoom, but the lens file has a valid mapping,
-				// evaluate the mapping at 0.0f. In this case, it is expected that the mapping will have 
-				// exactly one value in it, so warn the user if this is not the case.
-				if (LensFile->FocalLengthTable.GetFocusPointNum() > 0)
-				{
-					//Find FocusPoints involved for streamed Focus value. If there are, make sure there is only one zoom point in both of them or warn the user
-					if (LiveLinkCameraControllerUtils::HasSingleZoomPointInvolved(LensFileEvalData.Input.Focus, LensFile->FocalLengthTable.GetFocusPoints()) == false)
-					{
-						static const FName NAME_InvalidFocalLengthMappingWhenNotStreamed = "LiveLinkCamera_ZoomNotStreamedInvalidFocalLengthMapping";
-						FLiveLinkLog::WarningOnce(NAME_InvalidFocalLengthMappingWhenNotStreamed, FakedSubjectKey, TEXT("Problem applying FocalLength using subject '%s' and LensFile '%s'. Zoom wasn't streamed in and more than one focal length zoom point was found."), *FakedSubjectKey.SubjectName.ToString(), *LensFile->GetName());
-					}
-				}
-			}
-		}
-
-		//Verify distortion and image center tables 
-		{
-			if (StaticData->bIsFocusDistanceSupported == false)
-			{
-				//No focus : Look for single focus point in tables
-				if (LensFile->DataMode == ELensDataMode::Parameters && LensFile->DistortionTable.GetFocusPointNum() > 0 && LensFile->DistortionTable.GetFocusPointNum() != 1)
-				{
-					static const FName NAME_InvalidDistortionMappingWhenNotStreamed = "LiveLinkCamera_FocusNotStreamedInvalidDistortionMapping";
-					FLiveLinkLog::WarningOnce(NAME_InvalidDistortionMappingWhenNotStreamed, FakedSubjectKey, TEXT("Potential problem evaluating Distortion data using subject '%s' and LensFile '%s'. Focus wasn't streamed in and more than one Distortion focus point was found."), *FakedSubjectKey.SubjectName.ToString(), *LensFile->GetName());
-				}
-
-				if (LensFile->ImageCenterTable.GetFocusPointNum() > 0 && LensFile->ImageCenterTable.GetFocusPointNum() != 1)
-				{
-					static const FName NAME_InvalidImageCenterMappingWhenNotStreamed = "LiveLinkCamera_FocusNotStreamedInvalidImageCenterMapping";
-					FLiveLinkLog::WarningOnce(NAME_InvalidImageCenterMappingWhenNotStreamed, FakedSubjectKey, TEXT("Potential problem evaluating ImageCenter data using subject '%s' and LensFile '%s'. Focus wasn't streamed in and more than one ImageCenter focus point was found."), *FakedSubjectKey.SubjectName.ToString(), *LensFile->GetName());
-				}
-
-				if (LensFile->DataMode == ELensDataMode::STMap && LensFile->STMapTable.GetFocusPointNum() > 0 && LensFile->STMapTable.GetFocusPointNum() != 1)
-				{
-					static const FName NAME_InvalidSTMapMappingWhenNotStreamed = "LiveLinkCamera_FocusNotStreamedInvalidSTMapMapping";
-					FLiveLinkLog::WarningOnce(NAME_InvalidSTMapMappingWhenNotStreamed, FakedSubjectKey, TEXT("Potential problem evaluating STMap data using subject '%s' and LensFile '%s'. Focus wasn't streamed in and more than one STMap focus point was found."), *FakedSubjectKey.SubjectName.ToString(), *LensFile->GetName());
-				}
-			}
-
-			if(StaticData->bIsFocalLengthSupported == false)
-			{
-				//Zoom not supported -> Look for one zoom point in involved focus points
-				{
-					//Find FocusPoints involved for streamed Focus value. If there are, make sure there is only one zoom point in both of them or warn the user
-					if (LensFile->DataMode == ELensDataMode::Parameters
-						&& LensFile->DistortionTable.GetFocusPointNum() > 0 
-						&& LiveLinkCameraControllerUtils::HasSingleZoomPointInvolved<FDistortionFocusPoint>(LensFileEvalData.Input.Focus, LensFile->DistortionTable.GetFocusPoints()) == false)
-					{
-						static const FName NAME_InvalidDistortionMappingWhenNotStreamed = "LiveLinkCamera_ZoomNotStreamedInvalidDistortionMapping";
-						FLiveLinkLog::WarningOnce(NAME_InvalidDistortionMappingWhenNotStreamed, FakedSubjectKey, TEXT("Potential problem evaluating Distortion data using subject '%s' and LensFile '%s'. Zoom wasn't streamed in and more than one Distortion zoom point was found."), *FakedSubjectKey.SubjectName.ToString(), *LensFile->GetName());
-					}
-
-					if (LensFile->DataMode == ELensDataMode::STMap
-						&& LensFile->STMapTable.GetFocusPointNum() > 0
-						&& LiveLinkCameraControllerUtils::HasSingleZoomPointInvolved<FSTMapFocusPoint>(LensFileEvalData.Input.Focus, LensFile->STMapTable.GetFocusPoints()) == false)
-					{
-						static const FName NAME_InvalidSTMapMappingWhenNotStreamed = "LiveLinkCamera_ZoomNotStreamedInvalidSTMapMapping";
-						FLiveLinkLog::WarningOnce(NAME_InvalidSTMapMappingWhenNotStreamed, FakedSubjectKey, TEXT("Potential problem evaluating STMap data using subject '%s' and LensFile '%s'. Zoom wasn't streamed in and more than one STMap zoom point was found."), *FakedSubjectKey.SubjectName.ToString(), *LensFile->GetName());
-					}
-
-					if (LensFile->ImageCenterTable.GetFocusPointNum() > 0
-						&& LiveLinkCameraControllerUtils::HasSingleZoomPointInvolved<FImageCenterFocusPoint>(LensFileEvalData.Input.Focus, LensFile->ImageCenterTable.GetFocusPoints()) == false)
-					{
-						static const FName NAME_InvalidImageCenterMappingWhenNotStreamed = "LiveLinkCamera_ZoomNotStreamedInvalidImageCenterMapping";
-						FLiveLinkLog::WarningOnce(NAME_InvalidImageCenterMappingWhenNotStreamed, FakedSubjectKey, TEXT("Potential problem evaluating ImageCenter data using subject '%s' and LensFile '%s'. Zoom wasn't streamed in and more than one ImageCenter zoom point was found."), *FakedSubjectKey.SubjectName.ToString(), *LensFile->GetName());
-					}
-				}
-			}
-		}
 	}
 }
 
@@ -672,7 +312,7 @@ void ULiveLinkCameraController::PostLoad()
 	const int32 UE5MainVersion = GetLinkerCustomVersion(FUE5MainStreamObjectVersion::GUID);
 	if (UE5MainVersion < FUE5MainStreamObjectVersion::LensComponentNodalOffset)
 	{
-		// LensFile evaluation for nodal offset has moved to the Lens component. In order for existing camera actors to continue functioning as before,
+		// LensFile and distortion evaluation for nodal offset has moved to the Lens component. In order for existing camera actors to continue functioning as before,
 		// we have to migrate some properties from the camera controller to the lens component. If the owning actor does not have already have a lens 
 		// component then we will instance a new one.
 		AActor* const CameraActor = GetOuterActor();
@@ -698,6 +338,20 @@ void ULiveLinkCameraController::PostLoad()
 					{
 						LensComponent->SetApplyNodalOffsetOnTick(bApplyNodalOffset_DEPRECATED);
 					}
+
+					// Copy the overscan multiplier settings to the lens component
+					if (bScaleOverscan_DEPRECATED)
+					{
+						LensComponent->SetOverscanMultiplier(OverscanMultiplier_DEPRECATED);
+					}
+
+					// Copy the cropped filmback settings
+					if (bUseCroppedFilmback_DEPRECATED)
+					{
+						LensComponent->SetFilmbackOverrideSetting(EFilmbackOverrideSource::CroppedFilmbackSetting);
+						LensComponent->SetCroppedFilmback(CroppedFilmback_DEPRECATED);
+					}
+
 					PRAGMA_ENABLE_DEPRECATION_WARNINGS
 				}
 			}
@@ -711,9 +365,23 @@ void ULiveLinkCameraController::PostLoad()
 					CameraActor->AddInstanceComponent(LensComponent);
 					LensComponent->RegisterComponent();
 
-					LensComponent->SetLensFile(LensFilePicker);
+					LensComponent->SetLensFilePicker(LensFilePicker);
+
 					PRAGMA_DISABLE_DEPRECATION_WARNINGS
 					LensComponent->SetApplyNodalOffsetOnTick(bApplyNodalOffset_DEPRECATED);
+
+					// Copy the overscan multiplier settings to the lens component
+					if (bScaleOverscan_DEPRECATED)
+					{
+						LensComponent->SetOverscanMultiplier(OverscanMultiplier_DEPRECATED);
+					}
+
+					// Copy the cropped filmback settings
+					if (bUseCroppedFilmback_DEPRECATED)
+					{
+						LensComponent->SetFilmbackOverrideSetting(EFilmbackOverrideSource::CroppedFilmbackSetting);
+						LensComponent->SetCroppedFilmback(CroppedFilmback_DEPRECATED);
+					}
 					PRAGMA_ENABLE_DEPRECATION_WARNINGS
 				}
 			}
diff --git a/Engine/Plugins/VirtualProduction/LiveLinkCamera/Source/LiveLinkCamera/Public/LiveLinkCameraController.h b/Engine/Plugins/VirtualProduction/LiveLinkCamera/Source/LiveLinkCamera/Public/LiveLinkCameraController.h
--- a/Engine/Plugins/VirtualProduction/LiveLinkCamera/Source/LiveLinkCamera/Public/LiveLinkCameraController.h
+++ b/Engine/Plugins/VirtualProduction/LiveLinkCamera/Source/LiveLinkCamera/Public/LiveLinkCameraController.h
@@ -65,18 +65,10 @@ public:
 	virtual void Tick(float DeltaTime, const FLiveLinkSubjectFrameData& SubjectData) override;
 	virtual bool IsRoleSupported(const TSubclassOf<ULiveLinkRole>& RoleToSupport) override;
 	virtual TSubclassOf<UActorComponent> GetDesiredComponentClass() const override;
-	virtual void SetAttachedComponent(UActorComponent* ActorComponent) override;
-	virtual void Cleanup() override;
 	//~ End ULiveLinkControllerBase interface
 
 	//~ Begin UObject interface
-	virtual void PostDuplicate(bool bDuplicateForPIE) override;
-	virtual void PostEditImport() override;
 	virtual void PostLoad() override;
-#if WITH_EDITOR
-	virtual bool CanEditChange(const FProperty* InProperty) const override;
-	virtual void PostEditChangeProperty(struct FPropertyChangedEvent& PropertyChangedEvent) override;
-#endif	
 	//~ End UObject interface
 
 	/** Returns a const reference to input data used to evaluate the lens file */
@@ -92,7 +84,7 @@ protected:
 	UE_DEPRECATED(5.1, "This function has been deprecated. Nodal Offset is now applied by the LensComponent")
 	void ApplyNodalOffset(ULensFile* LensFile, UCineCameraComponent* CineCameraComponent);
 
-	/** Update distortion state */
+	UE_DEPRECATED(5.1, "This function has been deprecated. Distortion is now applied by the LensComponent")
 	void ApplyDistortion(ULensFile* LensFile, UCineCameraComponent* CineCameraComponent, const FLiveLinkCameraStaticData* StaticData, const FLiveLinkCameraFrameData* FrameData);
 
 	UE_DEPRECATED(5.1, "The use of this callback by this class has been deprecated and it is no longer registered. You can register your own delegate with FWorldDelegates::OnWorldPostActorTick")
@@ -123,39 +115,29 @@ public:
 	UPROPERTY(EditAnywhere, Category = "Camera Calibration")
 	FLensFilePicker LensFilePicker;
 
-	/** Whether to use the cropped filmback setting to drive the filmback of the attached camera component */
-	UPROPERTY(EditAnywhere, Category = "Camera Calibration")
-	bool bUseCroppedFilmback = false;
-
-	/** 
-	 * If a LensFile is being evaluated, the filmback saved in that LensFile will drive the attached camera component to ensure correct calibration.
-	 * If bUseCroppedFilmback is true, this value will be applied to the camera component and used to evaluate the LensFile instead.
-	 */
-	UPROPERTY(EditAnywhere, Category = "Camera Calibration")
-	FCameraFilmbackSettings CroppedFilmback;
-
 #if WITH_EDITORONLY_DATA
 	UE_DEPRECATED(5.1, "This property has been deprecated. The LiveLink camera controller no longer applies nodal offset.")
 	UPROPERTY()
 	bool bApplyNodalOffset_DEPRECATED = true;
-#endif //WITH_EDITORONLY_DATA
 
-	/** Whether to scale the computed overscan by the overscan percentage */
-	UPROPERTY(BlueprintReadWrite, Category = "Camera Calibration")
-	bool bScaleOverscan = false;
+	UE_DEPRECATED(5.1, "This property has been deprecated. Use the corresponding setting in the LensComponent instead.")
+	UPROPERTY()
+	bool bUseCroppedFilmback_DEPRECATED = false;
 
-	/** The percentage of the computed overscan that should be applied to the target camera */
-	UPROPERTY(EditAnywhere, BlueprintReadWrite, Category = "Camera Calibration", meta = (EditCondition = "bScaleOverscan", ClampMin = "0.0", ClampMax = "2.0"))
-	float OverscanMultiplier = 1.0f;
+	UE_DEPRECATED(5.1, "This property has been deprecated. Use the corresponding setting in the LensComponent instead.")
+	UPROPERTY()
+	FCameraFilmbackSettings CroppedFilmback_DEPRECATED;
 
-protected:
-	/** Cached distortion handler associated with attached camera component */
-	UPROPERTY(Transient)
-	ULensDistortionModelHandlerBase* LensDistortionHandler = nullptr;
+	UE_DEPRECATED(5.1, "This property has been deprecated. Use the corresponding setting in the LensComponent instead.")
+	UPROPERTY()
+	bool bScaleOverscan_DEPRECATED = false;
 
-	/** Unique identifier representing the source of distortion data */
-	UPROPERTY(DuplicateTransient)
-	FGuid DistortionProducerID;
+	UE_DEPRECATED(5.1, "This property has been deprecated. Use the corresponding setting in the LensComponent instead.")
+	UPROPERTY()
+	float OverscanMultiplier_DEPRECATED = 1.0f;
+#endif //WITH_EDITORONLY_DATA
+
+protected:
 
 #if WITH_EDITORONLY_DATA
 	UE_DEPRECATED(5.1, "This property has been deprecated. The LiveLink camera controller no longer applies nodal offset.")
@@ -165,6 +147,14 @@ protected:
 	UE_DEPRECATED(5.1, "This property has been deprecated. The LiveLink camera controller no longer applies nodal offset.")
 	UPROPERTY()
 	FVector OriginalCameraLocation_DEPRECATED;
+
+	UE_DEPRECATED(5.1, "This property has been deprecated. The LiveLink camera controller no longer evaluated distortion.")
+	UPROPERTY(Transient)
+	ULensDistortionModelHandlerBase* LensDistortionHandler_DEPRECATED = nullptr;
+
+	UE_DEPRECATED(5.1, "This property has been deprecated. The LiveLink camera controller no longer evaluated distortion.")
+	UPROPERTY(DuplicateTransient)
+	FGuid DistortionProducerID_DEPRECATED;
 #endif //WITH_EDITORONLY_DATA
 
 	/** Used to control which data from LiveLink is actually applied to camera */
@@ -182,8 +172,5 @@ protected:
 	FLensFileEvalData LensFileEvalData;
 
 private:
-	//Last values used to detect changes made by the user and update our original caches
-	FCameraFilmbackSettings LastFilmback;
-
 	double LastLensTableVerificationTimestamp = 0.0;
 };
diff --git a/Engine/Plugins/VirtualProduction/LiveLinkLens/Source/LiveLinkLens/Private/LiveLinkLensController.cpp b/Engine/Plugins/VirtualProduction/LiveLinkLens/Source/LiveLinkLens/Private/LiveLinkLensController.cpp
--- a/Engine/Plugins/VirtualProduction/LiveLinkLens/Source/LiveLinkLens/Private/LiveLinkLensController.cpp
+++ b/Engine/Plugins/VirtualProduction/LiveLinkLens/Source/LiveLinkLens/Private/LiveLinkLensController.cpp
@@ -3,18 +3,29 @@
 #include "LiveLinkLensController.h"
 
 #include "CameraCalibrationSubsystem.h"
-#include "CineCameraComponent.h"
 #include "Engine/Engine.h"
-#include "Engine/World.h"
+#include "LensComponent.h"
 #include "LiveLinkLensRole.h"
 #include "LiveLinkLensTypes.h"
+#include "UObject/UE5MainStreamObjectVersion.h"
 
+void ULiveLinkLensController::OnEvaluateRegistered()
+{
+	Super::OnEvaluateRegistered();
+	SetupLensComponent();
+}
+
+void ULiveLinkLensController::SetAttachedComponent(UActorComponent* ActorComponent)
+{
+	Super::SetAttachedComponent(ActorComponent);
+	SetupLensComponent();
+}
 
-ULiveLinkLensController::ULiveLinkLensController() 
+void ULiveLinkLensController::SetupLensComponent()
 {
-	if (!HasAnyFlags(RF_ArchetypeObject | RF_ClassDefaultObject))
+	if (ULensComponent* const LensComponent = Cast<ULensComponent>(GetAttachedComponent()))
 	{
-		DistortionProducerID = FGuid::NewGuid();
+		LensComponent->SetDistortionSource(EDistortionSource::LiveLinkLensSubject);
 	}
 }
 
@@ -25,45 +36,20 @@ void ULiveLinkLensController::Tick(float DeltaTime, const FLiveLinkSubjectFrameD
 
 	if (StaticData && FrameData)
 	{
-		if (UCineCameraComponent* const CineCameraComponent = Cast<UCineCameraComponent>(GetAttachedComponent()))
+		if (ULensComponent* const LensComponent = Cast<ULensComponent>(GetAttachedComponent()))
 		{
 			UCameraCalibrationSubsystem* const SubSystem = GEngine->GetEngineSubsystem<UCameraCalibrationSubsystem>();
 			const TSubclassOf<ULensModel> LensModel = SubSystem->GetRegisteredLensModel(StaticData->LensModel);
+			LensComponent->SetLensModel(LensModel);
 
-			const FName SubjectName = SelectedSubject.Subject.Name;
-			const FText RoleName = SelectedSubject.Role->GetDefaultObject<ULiveLinkRole>()->GetDisplayName();
-			const FString HandlerDisplayName = FString::Format(TEXT("'{0}' LiveLink Subject ({1} Role)"), { SubjectName.ToString(), RoleName.ToString() });
-
-			FDistortionHandlerPicker DistortionHandlerPicker = { CineCameraComponent, DistortionProducerID, HandlerDisplayName };
-			LensDistortionHandler = SubSystem->FindOrCreateDistortionModelHandler(DistortionHandlerPicker, LensModel);
-			
-			if (LensDistortionHandler)
-			{
-				// Update the lens distortion handler with the latest frame of data from the LiveLink source
-				FLensDistortionState DistortionState;
-
-				DistortionState.DistortionInfo.Parameters = FrameData->DistortionParameters;
-				DistortionState.FocalLengthInfo.FxFy = FrameData->FxFy;
-				DistortionState.ImageCenter.PrincipalPoint = FrameData->PrincipalPoint;
-
-				//Update the distortion state based on incoming LL data.
-				LensDistortionHandler->SetDistortionState(DistortionState);
+			// Update the lens distortion handler with the latest frame of data from the LiveLink source
+			FLensDistortionState DistortionState;
 
-				//Recompute overscan factor for the distortion state
-				float OverscanFactor = LensDistortionHandler->ComputeOverscanFactor();
+			DistortionState.DistortionInfo.Parameters = FrameData->DistortionParameters;
+			DistortionState.FocalLengthInfo.FxFy = FrameData->FxFy;
+			DistortionState.ImageCenter.PrincipalPoint = FrameData->PrincipalPoint;
 
-				if (bScaleOverscan)
-				{
-					OverscanFactor = ((OverscanFactor - 1.0f) * OverscanMultiplier) + 1.0f;
-				}
-				LensDistortionHandler->SetOverscanFactor(OverscanFactor);
-
-				//Make sure the displacement map is up to date
-				LensDistortionHandler->ProcessCurrentDistortion();
-			}
-
-			// Track changes to the cine camera's focal length for consumers of distortion data
-			SubSystem->UpdateOriginalFocalLength(CineCameraComponent, CineCameraComponent->CurrentFocalLength);
+			LensComponent->SetDistortionState(DistortionState);
 		}
 	}
 }
@@ -75,39 +61,73 @@ bool ULiveLinkLensController::IsRoleSupported(const TSubclassOf<ULiveLinkRole>&
 
 TSubclassOf<UActorComponent> ULiveLinkLensController::GetDesiredComponentClass() const
 {
-	return UCineCameraComponent::StaticClass();
+	return ULensComponent::StaticClass();
 }
 
-void ULiveLinkLensController::Cleanup()
+
+void ULiveLinkLensController::PostLoad()
 {
-	if (UCineCameraComponent* const CineCameraComponent = Cast<UCineCameraComponent>(GetAttachedComponent()))
+	Super::PostLoad();
+
+#if WITH_EDITOR
+	const int32 UE5MainVersion = GetLinkerCustomVersion(FUE5MainStreamObjectVersion::GUID);
+	if (UE5MainVersion < FUE5MainStreamObjectVersion::LensComponentDistortion)
 	{
-		if (UCameraCalibrationSubsystem* SubSystem = GEngine->GetEngineSubsystem<UCameraCalibrationSubsystem>())
+		// Distortion evaluation has moved to the Lens component. In order for existing camera actors to continue functioning as before,
+		// we have to migrate some properties from the lens controller to the lens component. If the owning actor does not have already have a lens 
+		// component then we will instance a new one.
+		AActor* const CameraActor = GetOuterActor();
+		CameraActor->ConditionalPostLoad();
+		if (CameraActor->GetWorld())
 		{
-			SubSystem->UnregisterDistortionModelHandler(CineCameraComponent, LensDistortionHandler);
-		}
-	}
-}
+			TInlineComponentArray<ULensComponent*> LensComponents;
+			CameraActor->GetComponents(LensComponents);
 
-void ULiveLinkLensController::PostDuplicate(bool bDuplicateForPIE)
-{
-	Super::PostDuplicate(bDuplicateForPIE);
+			// Loop through any existing lens components to find one whose distortion source was this lens controller
+			bool bFoundMatchingLensComponent = false;
+			for (ULensComponent* LensComponent : LensComponents)
+			{
+				PRAGMA_DISABLE_DEPRECATION_WARNINGS
+				FDistortionHandlerPicker DeprecatedPicker = LensComponent->GetDistortionHandlerPicker();
+				if (DeprecatedPicker.DistortionProducerID == DistortionProducerID_DEPRECATED)
+				{
+					bFoundMatchingLensComponent = true;
+
+					// Preserve the previous "bApplyDistortion" setting from the existing component
+					const bool bPreviousDistortionSetting = LensComponent->ShouldApplyDistortion();
+					SetAttachedComponent(LensComponent);
+					LensComponent->SetApplyDistortion(bPreviousDistortionSetting);
+					
+					// Copy the overscan multiplier settings to the lens component
+					if (bScaleOverscan_DEPRECATED)
+					{
+						LensComponent->SetOverscanMultiplier(OverscanMultiplier_DEPRECATED);
+					}
+				}
+				PRAGMA_ENABLE_DEPRECATION_WARNINGS
+			}
 
-	// When this controller is duplicated (e.g. for PIE), the duplicated controller needs a new unique ID
-	if (!DistortionProducerID.IsValid())
-	{
-		DistortionProducerID = FGuid::NewGuid();
-	}
-}
+			// If no matching existing lens components were found, create a new one and initialize it with the correct LensFile
+			if (!bFoundMatchingLensComponent)
+			{
+				ULensComponent* LensComponent = NewObject<ULensComponent>(CameraActor, MakeUniqueObjectName(this, ULensComponent::StaticClass(), TEXT("Lens")));
+				CameraActor->AddInstanceComponent(LensComponent);
+				LensComponent->RegisterComponent();
 
-void ULiveLinkLensController::PostEditImport()
-{
-	Super::PostEditImport();
+				SetAttachedComponent(LensComponent);
 
-	// PostDuplicate is not called on components during actor duplication, such as alt-drag and copy-paste, so PostEditImport covers those duplication cases
-	// When this controller is duplicated in those cases, the duplicated controller needs a new unique ID
-	if (!DistortionProducerID.IsValid())
-	{
-		DistortionProducerID = FGuid::NewGuid();
+				// If there was no existing lens component, then the camera would not have had distortion (from this controller) applied to it
+				LensComponent->SetApplyDistortion(false);
+
+				PRAGMA_DISABLE_DEPRECATION_WARNINGS
+				// Copy the overscan multiplier settings to the lens component
+				if (bScaleOverscan_DEPRECATED)
+				{
+					LensComponent->SetOverscanMultiplier(OverscanMultiplier_DEPRECATED);
+				}
+				PRAGMA_ENABLE_DEPRECATION_WARNINGS
+			}
+		}
 	}
+#endif //WITH_EDITOR
 }
diff --git a/Engine/Plugins/VirtualProduction/LiveLinkLens/Source/LiveLinkLens/Public/LiveLinkLensController.h b/Engine/Plugins/VirtualProduction/LiveLinkLens/Source/LiveLinkLens/Public/LiveLinkLensController.h
--- a/Engine/Plugins/VirtualProduction/LiveLinkLens/Source/LiveLinkLens/Public/LiveLinkLensController.h
+++ b/Engine/Plugins/VirtualProduction/LiveLinkLens/Source/LiveLinkLens/Public/LiveLinkLensController.h
@@ -17,34 +17,39 @@ class LIVELINKLENS_API ULiveLinkLensController : public ULiveLinkControllerBase
 	GENERATED_BODY()
 
 public:
-	ULiveLinkLensController();
-
 	//~ Begin ULiveLinkControllerBase interface
+	virtual void OnEvaluateRegistered() override;
 	virtual void Tick(float DeltaTime, const FLiveLinkSubjectFrameData& SubjectData) override;
 	virtual bool IsRoleSupported(const TSubclassOf<ULiveLinkRole>& RoleToSupport) override;
 	virtual TSubclassOf<UActorComponent> GetDesiredComponentClass() const override;
-	virtual void Cleanup() override;
+	virtual void SetAttachedComponent(UActorComponent* ActorComponent) override;
 	//~ End ULiveLinkControllerBase interface
 
-	//~ Begin UObject Interface
-	virtual void PostDuplicate(bool bDuplicateForPIE) override;
-	virtual void PostEditImport() override;
-	//~ End UObject Interface
+	//~ Begin UObject interface
+	virtual void PostLoad() override;
+	//~ End UObject interface
+
+private:
+	/** Sets the evaluation mode of the attached LensComponent to use LiveLink and enable applying distortion */
+	void SetupLensComponent();
 
 protected:
-	/** Cached distortion handler associated with attached camera component */
-	UPROPERTY(Transient)
-	ULensDistortionModelHandlerBase* LensDistortionHandler = nullptr;
 
-	/** Unique identifier representing the source of distortion data */
+#if WITH_EDITORONLY_DATA
+	UE_DEPRECATED(5.1, "This property has been deprecated. Distortion is handled directly by controlled Lens Component.")
+ 	UPROPERTY(Transient)
+ 	ULensDistortionModelHandlerBase* LensDistortionHandler_DEPRECATED = nullptr;
+
+	UE_DEPRECATED(5.1, "This property has been deprecated and is no longer used.")
 	UPROPERTY(DuplicateTransient)
-	FGuid DistortionProducerID;
+	FGuid DistortionProducerID_DEPRECATED;
 
-	/** Whether to scale the computed overscan by the overscan percentage */
-	UPROPERTY(BlueprintReadWrite, Category = "Camera Calibration")
-	bool bScaleOverscan = false;
+	UE_DEPRECATED(5.1, "This property has been deprecated. Overscan multipliers should now be applied from a Lens Component.")
+	UPROPERTY()
+	bool bScaleOverscan_DEPRECATED = false;
 
-	/** The percentage of the computed overscan that should be applied to the target camera */
-	UPROPERTY(EditAnywhere, BlueprintReadWrite, Category = "Camera Calibration", meta = (EditCondition = "bScaleOverscan", ClampMin = "0.0", ClampMax = "2.0"))
-	float OverscanMultiplier = 1.0f;
-};
\ No newline at end of file
+	UE_DEPRECATED(5.1, "This property has been deprecated. Overscan multipliers should now be applied from a Lens Component.")
+	UPROPERTY()
+	float OverscanMultiplier_DEPRECATED = 1.0f;
+#endif //WITH_EDITORONLY_DATA
+};
diff --git a/Engine/Source/Runtime/Core/Public/UObject/UE5MainStreamObjectVersion.h b/Engine/Source/Runtime/Core/Public/UObject/UE5MainStreamObjectVersion.h
--- a/Engine/Source/Runtime/Core/Public/UObject/UE5MainStreamObjectVersion.h
+++ b/Engine/Source/Runtime/Core/Public/UObject/UE5MainStreamObjectVersion.h
@@ -256,6 +256,9 @@ struct CORE_API FUE5MainStreamObjectVersion
 		// Made ManagedArrayCollection::FValueType::Value always serialize when FValueType is
 		ManagedArrayCollectionAlwaysSerializeValue,
 
+		// Moving all distortion handling to Lens Component
+		LensComponentDistortion,
+
 		// -----<new versions can be added above this line>-------------------------------------------------
 		VersionPlusOne,
 		LatestVersion = VersionPlusOne - 1
