[CameraCalibration] Add custom track recorder for LensComponent to support caching the LensFile, the LensFile evaluation inputs, and the distortion state (parameters and intrinsics), in order to properly support evaluation and application of distortion on playback of a level sequence.

diff --git a/Engine/Plugins/VirtualProduction/CameraCalibrationCore/CameraCalibrationCore.uplugin b/Engine/Plugins/VirtualProduction/CameraCalibrationCore/CameraCalibrationCore.uplugin
--- a/Engine/Plugins/VirtualProduction/CameraCalibrationCore/CameraCalibrationCore.uplugin
+++ b/Engine/Plugins/VirtualProduction/CameraCalibrationCore/CameraCalibrationCore.uplugin
@@ -24,6 +24,16 @@
 			"Name": "CameraCalibrationCoreEditor",
 			"Type": "Editor",
 			"LoadingPhase": "Default"
+		},
+		{
+			"Name": "CameraCalibrationCoreMovieScene",
+			"Type": "Runtime",
+			"LoadingPhase": "Default"
+		},
+		{
+			"Name": "CameraCalibrationCoreSequencer",
+			"Type": "UncookedOnly",
+			"LoadingPhase": "Default"
 		}
 	],
 	"Plugins": [
diff --git a/Engine/Plugins/VirtualProduction/CameraCalibrationCore/Source/CameraCalibrationCoreMovieScene/CameraCalibrationCoreMovieScene.build.cs b/Engine/Plugins/VirtualProduction/CameraCalibrationCore/Source/CameraCalibrationCoreMovieScene/CameraCalibrationCoreMovieScene.build.cs
new file mode 100644
--- /dev/null
+++ b/Engine/Plugins/VirtualProduction/CameraCalibrationCore/Source/CameraCalibrationCoreMovieScene/CameraCalibrationCoreMovieScene.build.cs
@@ -0,0 +1,40 @@
+// Copyright Epic Games, Inc. All Rights Reserved.
+
+using UnrealBuildTool;
+
+public class CameraCalibrationCoreMovieScene : ModuleRules
+{
+	public CameraCalibrationCoreMovieScene(ReadOnlyTargetRules Target) : base(Target)
+	{
+		PCHUsage = ModuleRules.PCHUsageMode.UseExplicitOrSharedPCHs;
+
+		PublicIncludePaths.AddRange(
+			new string[] {
+			}
+			);
+				
+		
+		PrivateIncludePaths.AddRange(
+			new string[] {
+			}
+		);
+			
+		
+		PublicDependencyModuleNames.AddRange(
+			new string[]
+			{
+			}
+		);
+		
+		PrivateDependencyModuleNames.AddRange(
+			new string[]
+			{
+ 				"CameraCalibrationCore", 
+ 				"Core",
+ 				"CoreUObject",
+ 				"Engine",
+ 				"MovieScene",
+			}
+		);
+	}
+}
diff --git a/Engine/Plugins/VirtualProduction/CameraCalibrationCore/Source/CameraCalibrationCoreMovieScene/Private/CameraCalibrationCoreMovieScene.cpp b/Engine/Plugins/VirtualProduction/CameraCalibrationCore/Source/CameraCalibrationCoreMovieScene/Private/CameraCalibrationCoreMovieScene.cpp
new file mode 100644
--- /dev/null
+++ b/Engine/Plugins/VirtualProduction/CameraCalibrationCore/Source/CameraCalibrationCoreMovieScene/Private/CameraCalibrationCoreMovieScene.cpp
@@ -0,0 +1,5 @@
+// Copyright Epic Games, Inc. All Rights Reserved.
+
+#include "Modules/ModuleManager.h"
+
+IMPLEMENT_MODULE(FDefaultModuleImpl, CameraCalibrationCoreMovieScene);
diff --git a/Engine/Plugins/VirtualProduction/CameraCalibrationCore/Source/CameraCalibrationCoreMovieScene/Private/MovieSceneLensComponentSection.cpp b/Engine/Plugins/VirtualProduction/CameraCalibrationCore/Source/CameraCalibrationCoreMovieScene/Private/MovieSceneLensComponentSection.cpp
new file mode 100644
--- /dev/null
+++ b/Engine/Plugins/VirtualProduction/CameraCalibrationCore/Source/CameraCalibrationCoreMovieScene/Private/MovieSceneLensComponentSection.cpp
@@ -0,0 +1,277 @@
+// Copyright Epic Games, Inc. All Rights Reserved.
+
+#include "MovieSceneLensComponentSection.h"
+
+#include "Channels/MovieSceneChannelProxy.h"
+#include "Evaluation/MovieScenePreAnimatedState.h"
+#include "Evaluation/PreAnimatedState/MovieScenePreAnimatedCaptureSource.h"
+#include "IMovieScenePlayer.h"
+#include "MovieSceneExecutionToken.h"
+#include "MovieSceneTimeHelpers.h"
+
+#define LOCTEXT_NAMESPACE "MovieSceneLensComponentSection"
+
+struct FPreAnimatedLensFileTokenProducer : IMovieScenePreAnimatedTokenProducer
+{
+	virtual IMovieScenePreAnimatedTokenPtr CacheExistingState(UObject& Object) const
+	{
+		struct FLensFileToken : IMovieScenePreAnimatedToken
+		{
+			FLensFileToken(ULensComponent* LensComponent)
+			{
+				// Save the original values of these settings
+				FLensFilePicker LensFilePicker = LensComponent->GetLensFilePicker();
+				bOriginalUseDefaultLensFile = LensFilePicker.bUseDefaultLensFile;
+				OriginalLensFile = LensFilePicker.GetLensFile();
+				OriginalEvalMode = LensComponent->GetFIZEvaluationMode();
+			}
+
+			virtual void RestoreState(UObject& Object, const UE::MovieScene::FRestoreStateParams& Params)
+			{
+				if (ULensComponent* LensComponent = CastChecked<ULensComponent>(&Object))
+				{
+					// Restore the original values of these settings
+					LensComponent->SetLensFilePicker({ bOriginalUseDefaultLensFile, OriginalLensFile.Get() });
+					LensComponent->SetFIZEvaluationMode(OriginalEvalMode);
+				}
+			}
+
+			bool bOriginalUseDefaultLensFile = false;
+			TWeakObjectPtr<ULensFile> OriginalLensFile;
+			EFIZEvaluationMode OriginalEvalMode = EFIZEvaluationMode::UseRecordedValues;
+		};
+
+		return FLensFileToken(CastChecked<ULensComponent>(&Object));
+	}
+	static FMovieSceneAnimTypeID GetAnimTypeID()
+	{
+		return TMovieSceneAnimTypeID<FPreAnimatedLensFileTokenProducer>();
+	}
+};
+
+void UMovieSceneLensComponentSection::Initialize(ULensComponent* Component)
+{
+	// Cache a pointer to the recorded component to retrieve its data later during RecordFrame()
+	RecordedComponent = Component;
+
+	// Duplicate the recorded LensFile asset and save the duplicate in this section
+ 	if (ULensFile* LensFile = Component->GetLensFile())
+ 	{
+ 		const FName DuplicateName = FName(*FString::Format(TEXT("{0}_Cached"), { LensFile->GetFName().ToString() }));
+ 		CachedLensFile = Cast<ULensFile>(StaticDuplicateObject(LensFile, this, DuplicateName, RF_AllFlags & ~RF_Transient));
+ 	}
+
+	// Create the channels that will be used to record the distortion state
+	CreateChannelProxy();
+}
+
+#if WITH_EDITOR
+void UMovieSceneLensComponentSection::AddChannelWithEditor(FMovieSceneChannelProxyData& ChannelProxyData, FMovieSceneFloatChannel& Channel, FText GroupName, FText ChannelName, int32 SortOrder)
+{
+	FMovieSceneChannelMetaData ChannelEditorData(FName(*(ChannelName.ToString())), ChannelName, GroupName);
+	ChannelEditorData.SortOrder = SortOrder;
+	ChannelEditorData.bCanCollapseToTrack = false;
+
+	ChannelProxyData.Add(Channel, ChannelEditorData, TMovieSceneExternalValue<float>());
+}
+#endif // WITH_EDITOR
+
+void UMovieSceneLensComponentSection::CreateChannelProxy()
+{
+	if (const ULensComponent* const LensComponent = RecordedComponent.Get())
+	{
+		FMovieSceneChannelProxyData Channels;
+
+		// Add channels for distortion parameters
+		if (TSubclassOf<ULensModel> LensModelClass = LensComponent->GetLensModel())
+		{
+			if (const ULensModel* const LensModel = LensModelClass->GetDefaultObject<ULensModel>())
+			{
+				const int32 ParameterCount = LensModel->GetNumParameters();
+				DistortionParameterChannels.SetNum(ParameterCount);
+
+#if WITH_EDITOR
+				TArray<FText> ParameterDisplayNames = LensModel->GetParameterDisplayNames();
+				const FText ParameterGroupName = LOCTEXT("ParameterGroupName", "Distortion Parameters");
+#endif // WITH_EDITOR
+
+				for (int32 Index = 0; Index < ParameterCount; ++Index)
+				{
+#if WITH_EDITOR
+					AddChannelWithEditor(Channels, DistortionParameterChannels[Index], ParameterGroupName, ParameterDisplayNames[Index], Index);
+#else
+					Channels.Add(DistortionParameterChannels[Index]);
+#endif // WITH_EDITOR
+				}
+			}
+		}
+
+		// Add channels for FxFy
+		FxFyChannels.SetNum(2);
+
+#if WITH_EDITOR
+		const FText FxFyGroupName = LOCTEXT("FxFyGroupName", "FxFy");
+		constexpr int32 FxFySortOrder = 1000;
+		AddChannelWithEditor(Channels, FxFyChannels[0], FxFyGroupName, LOCTEXT("FxChannelName", "Fx"), FxFySortOrder + 0);
+		AddChannelWithEditor(Channels, FxFyChannels[1], FxFyGroupName, LOCTEXT("FyChannelName", "Fy"), FxFySortOrder + 1);
+#else
+		Channels.Add(FxFyChannels[0]);
+		Channels.Add(FxFyChannels[1]);
+#endif // WITH_EDITOR
+
+		// Add channels for Image Center
+		ImageCenterChannels.SetNum(2);
+
+#if WITH_EDITOR
+		const FText ImageCenterGroupName = LOCTEXT("CxCyGroupName", "Image Center");
+		constexpr int32 ImageCenterSortOrder = 2000;
+		AddChannelWithEditor(Channels, ImageCenterChannels[0], ImageCenterGroupName, LOCTEXT("CxChannelName", "Cx"), ImageCenterSortOrder + 0);
+		AddChannelWithEditor(Channels, ImageCenterChannels[1], ImageCenterGroupName, LOCTEXT("CyChannelName", "Cy"), ImageCenterSortOrder + 1);
+#else
+		Channels.Add(ImageCenterChannels[0]);
+		Channels.Add(ImageCenterChannels[1]);
+#endif // WITH_EDITOR
+
+		ChannelProxy = MakeShared<FMovieSceneChannelProxy>(MoveTemp(Channels));
+	}
+}
+
+void UMovieSceneLensComponentSection::RecordFrame(FFrameNumber FrameNumber)
+{
+	if (const ULensComponent* const LensComponent = RecordedComponent.Get())
+	{
+		// Get the current distortion state of the lens component being recorded
+		const FLensDistortionState CurrentState = LensComponent->GetDistortionState();
+		const int32 ParameterCount = CurrentState.DistortionInfo.Parameters.Num();
+
+		// Record the current state of each distortion parameter
+		for (int32 Index = 0; Index < ParameterCount; ++Index)
+		{
+			DistortionParameterChannels[Index].GetData().AddKey(FrameNumber, FMovieSceneFloatValue(CurrentState.DistortionInfo.Parameters[Index]));
+		}
+
+		// Record the current state of FxFy
+		FxFyChannels[0].GetData().AddKey(FrameNumber, FMovieSceneFloatValue(CurrentState.FocalLengthInfo.FxFy[0]));
+		FxFyChannels[1].GetData().AddKey(FrameNumber, FMovieSceneFloatValue(CurrentState.FocalLengthInfo.FxFy[1]));
+
+		// Record the current state of the image center
+		ImageCenterChannels[0].GetData().AddKey(FrameNumber, FMovieSceneFloatValue(CurrentState.ImageCenter.PrincipalPoint[0]));
+		ImageCenterChannels[1].GetData().AddKey(FrameNumber, FMovieSceneFloatValue(CurrentState.ImageCenter.PrincipalPoint[1]));
+	}
+}
+
+void UMovieSceneLensComponentSection::Finalize()
+{
+	// Reduce keys on each of the channels that we recorded
+	FKeyDataOptimizationParams Params;
+
+	const int32 ParameterCount = DistortionParameterChannels.Num();
+	for (int32 Index = 0; Index < ParameterCount; ++Index)
+	{
+		UE::MovieScene::Optimize(&DistortionParameterChannels[Index], Params);
+	}
+
+	UE::MovieScene::Optimize(&FxFyChannels[0], Params);
+	UE::MovieScene::Optimize(&FxFyChannels[1], Params);
+	UE::MovieScene::Optimize(&ImageCenterChannels[0], Params);
+	UE::MovieScene::Optimize(&ImageCenterChannels[1], Params);
+}
+
+void UMovieSceneLensComponentSection::Begin(IMovieScenePlayer* Player, const UE::MovieScene::FEvaluationHookParams& Params) const
+{
+	FScopedPreAnimatedCaptureSource CaptureSource(&Player->PreAnimatedState, this, Params.SequenceID, EvalOptions.CompletionMode == EMovieSceneCompletionMode::RestoreState);
+
+	for (TWeakObjectPtr<> BoundObject : Player->FindBoundObjects(Params.ObjectBindingID, Params.SequenceID))
+	{
+		if (ULensComponent* LensComponent = Cast<ULensComponent>(BoundObject.Get()))
+		{
+			Player->SavePreAnimatedState(*LensComponent, FPreAnimatedLensFileTokenProducer::GetAnimTypeID(), FPreAnimatedLensFileTokenProducer());
+
+			// Override the existing lens file settings of the lens component to use the cached, duplicate lens file
+			LensComponent->SetFIZEvaluationMode(EFIZEvaluationMode::UseRecordedValues);
+
+			// TODO: Check that the number of parameters for the lens component's current lens model matches the number of recorded parameter channels. 
+		}
+	}
+}
+
+void UMovieSceneLensComponentSection::Update(IMovieScenePlayer* Player, const UE::MovieScene::FEvaluationHookParams& Params) const
+{
+	for (TWeakObjectPtr<> BoundObject : Player->FindBoundObjects(Params.ObjectBindingID, Params.SequenceID))
+	{
+		if (ULensComponent* LensComponent = Cast<ULensComponent>(BoundObject.Get()))
+		{
+			// Set the LensFile to use on the LensComponent, either the cached version or an override
+			if (OverrideLensFile)
+			{
+				LensComponent->SetLensFilePicker({ false, OverrideLensFile });
+			}
+			else
+			{
+				LensComponent->SetLensFilePicker({ false, CachedLensFile });
+			}
+
+			// Nodal offset is baked in to the transform of the camera, but users can use this flag to force the LensComponent to re-evaluate
+			// the LensFile and re-apply a new nodal offset instead.
+			if (bReapplyNodalOffset)
+			{
+				LensComponent->ReapplyNodalOffset();
+			}
+
+			// Get the end frame of the current range being evaluated
+			const TRange<FFrameNumber> PlaybackRange = Params.Context.GetFrameNumberRange();
+			const FFrameNumber ExclusiveEndFrame = UE::MovieScene::DiscreteExclusiveUpper(PlaybackRange);
+
+			// Get the current distortion state of the lens component, which will be overwritten with values from the channel (if present)
+			FLensDistortionState NewState = LensComponent->GetDistortionState();
+
+			// Update the state from the distortion parameter channels
+			const int32 ParameterCount = DistortionParameterChannels.Num();
+			NewState.DistortionInfo.Parameters.SetNum(ParameterCount);
+
+			for (int32 Index = 0; Index < ParameterCount; ++Index)
+			{
+				float NewValue = 0.0f;
+				if (DistortionParameterChannels[Index].Evaluate(ExclusiveEndFrame, NewValue))
+				{
+					NewState.DistortionInfo.Parameters[Index] = NewValue;
+				}
+			}
+
+			// Update the state from the FxFy channels
+			float NewValueFx = 0.0f;
+			if (FxFyChannels[0].Evaluate(ExclusiveEndFrame, NewValueFx))
+			{
+				NewState.FocalLengthInfo.FxFy[0] = NewValueFx;
+			}
+
+			float NewValueFy = 0.0f;
+			if (FxFyChannels[1].Evaluate(ExclusiveEndFrame, NewValueFy))
+			{
+				NewState.FocalLengthInfo.FxFy[1] = NewValueFy;
+			}
+
+			// Update the state from the ImageCenter channels
+			float NewValueCx = 0.0f;
+			if (ImageCenterChannels[0].Evaluate(ExclusiveEndFrame, NewValueCx))
+			{
+				NewState.ImageCenter.PrincipalPoint[0] = NewValueCx;
+			}
+
+			float NewValueCy = 0.0f;
+			if (ImageCenterChannels[1].Evaluate(ExclusiveEndFrame, NewValueCy))
+			{
+				NewState.ImageCenter.PrincipalPoint[1] = NewValueCy;
+			}
+
+			LensComponent->SetDistortionState(NewState);
+		}
+	}
+}
+
+void UMovieSceneLensComponentSection::End(IMovieScenePlayer* Player, const UE::MovieScene::FEvaluationHookParams& Params) const
+{
+	Player->RestorePreAnimatedState();
+}
+
+#undef LOCTEXT_NAMESPACE
diff --git a/Engine/Plugins/VirtualProduction/CameraCalibrationCore/Source/CameraCalibrationCoreMovieScene/Private/MovieSceneLensComponentTrack.cpp b/Engine/Plugins/VirtualProduction/CameraCalibrationCore/Source/CameraCalibrationCoreMovieScene/Private/MovieSceneLensComponentTrack.cpp
new file mode 100644
--- /dev/null
+++ b/Engine/Plugins/VirtualProduction/CameraCalibrationCore/Source/CameraCalibrationCoreMovieScene/Private/MovieSceneLensComponentTrack.cpp
@@ -0,0 +1,64 @@
+// Copyright Epic Games, Inc. All Rights Reserved.
+
+#include "MovieSceneLensComponentTrack.h"
+
+#include "MovieSceneLensComponentSection.h"
+
+#define LOCTEXT_NAMESPACE "MovieSceneLensComponentTrack"
+
+bool UMovieSceneLensComponentTrack::SupportsType(TSubclassOf<UMovieSceneSection> SectionClass) const
+{
+	return SectionClass == UMovieSceneLensComponentSection::StaticClass();
+}
+
+UMovieSceneSection* UMovieSceneLensComponentTrack::CreateNewSection()
+{
+	return NewObject<UMovieSceneSection>(this, UMovieSceneLensComponentSection::StaticClass(), MakeUniqueObjectName(this, UMovieSceneLensComponentSection::StaticClass(), TEXT("LensComponentSection")), RF_Transactional);
+}
+
+#if WITH_EDITORONLY_DATA
+FText UMovieSceneLensComponentTrack::GetDisplayName() const
+{
+	return LOCTEXT("LensComponentTrackDisplayName", "Lens Component");
+}
+#endif
+
+void UMovieSceneLensComponentTrack::AddSection(UMovieSceneSection& Section)
+{
+	Sections.Add(&Section);
+}
+
+const TArray<UMovieSceneSection*>& UMovieSceneLensComponentTrack::GetAllSections() const
+{
+	return Sections;
+}
+
+bool UMovieSceneLensComponentTrack::HasSection(const UMovieSceneSection& Section) const
+{
+	return Sections.Contains(&Section);
+}
+
+bool UMovieSceneLensComponentTrack::IsEmpty() const
+{
+	return Sections.IsEmpty();
+}
+
+void UMovieSceneLensComponentTrack::RemoveAllAnimationData()
+{
+	Sections.Empty();
+}
+
+void UMovieSceneLensComponentTrack::RemoveSection(UMovieSceneSection& Section)
+{
+	Sections.Remove(&Section);
+}
+
+void UMovieSceneLensComponentTrack::RemoveSectionAt(int32 SectionIndex)
+{
+	if (Sections.IsValidIndex(SectionIndex))
+	{
+		Sections.RemoveAt(SectionIndex);
+	}
+}
+
+#undef LOCTEXT_NAMESPACE
diff --git a/Engine/Plugins/VirtualProduction/CameraCalibrationCore/Source/CameraCalibrationCoreMovieScene/Public/MovieSceneLensComponentSection.h b/Engine/Plugins/VirtualProduction/CameraCalibrationCore/Source/CameraCalibrationCoreMovieScene/Public/MovieSceneLensComponentSection.h
new file mode 100644
--- /dev/null
+++ b/Engine/Plugins/VirtualProduction/CameraCalibrationCore/Source/CameraCalibrationCoreMovieScene/Public/MovieSceneLensComponentSection.h
@@ -0,0 +1,80 @@
+// Copyright Epic Games, Inc. All Rights Reserved.
+
+#pragma once
+
+#include "Sections/MovieSceneHookSection.h"
+
+#include "Channels/MovieSceneFloatChannel.h"
+#include "LensComponent.h"
+#include "LensFile.h"
+
+#include "MovieSceneLensComponentSection.generated.h"
+
+struct FMovieSceneChannelProxyData;
+
+/** Movie Scene section for Lens Component */
+UCLASS()
+class CAMERACALIBRATIONCOREMOVIESCENE_API UMovieSceneLensComponentSection : public UMovieSceneHookSection
+{
+	GENERATED_BODY()
+
+public:
+
+	UMovieSceneLensComponentSection(const FObjectInitializer& ObjectInitializer)
+		: UMovieSceneHookSection(ObjectInitializer) 
+	{ 
+		bRequiresRangedHook = true; 
+	}
+
+	//~ Begin UMovieSceneHookSection interface
+	virtual void Begin(IMovieScenePlayer* Player, const UE::MovieScene::FEvaluationHookParams& Params) const override;
+	virtual void Update(IMovieScenePlayer* Player, const UE::MovieScene::FEvaluationHookParams& Params) const override;
+	virtual void End(IMovieScenePlayer* Player, const UE::MovieScene::FEvaluationHookParams& Params) const override;
+	//~ End UMovieSceneHookSection interface
+
+	/** Initialize the section with the Lens Component whose properties it will record */
+	void Initialize(ULensComponent* Component);
+
+	/** Reduce keys in each of the recorded channels */
+	void Finalize();
+
+	/** Record the current value of each distortion state property to its appropriate track */
+	void RecordFrame(FFrameNumber FrameNumber);
+
+private:
+	/** Create the channels in which distortion state values will be recorded */
+	void CreateChannelProxy();
+
+#if WITH_EDITOR
+	void AddChannelWithEditor(FMovieSceneChannelProxyData& ChannelProxyData, FMovieSceneFloatChannel& Channel, FText GroupName, FText ChannelName, int32 SortOrder);
+#endif // WITH_EDITOR
+
+public:
+	/** If true, then every Update, the nodal offset will be re-evaluated on the lens component */
+	UPROPERTY(EditAnywhere, Category = "Camera Calibration")
+	bool bReapplyNodalOffset = false;
+
+	/** LensFile asset that should be used instead of the cached LensFile during playback */
+	UPROPERTY(EditAnywhere, Category = "Camera Calibration")
+	TObjectPtr<ULensFile> OverrideLensFile = nullptr;
+
+private:
+	/** Saved duplicate of the LensFile asset used by the recorded Lens Component at the time of recording */
+	UPROPERTY(VisibleAnywhere, Category="Camera Calibration")
+	TObjectPtr<ULensFile> CachedLensFile = nullptr;
+
+	/** Channels to store Distortion Parameter values (will be sized during initialization based on the LensComponent's LensModel) */
+	UPROPERTY()
+	TArray<FMovieSceneFloatChannel> DistortionParameterChannels;
+
+	/** Channels to store FxFy values */
+	UPROPERTY()
+	TArray<FMovieSceneFloatChannel> FxFyChannels;
+
+	/** Channels to store Image Center values */
+	UPROPERTY()
+	TArray<FMovieSceneFloatChannel> ImageCenterChannels;
+
+private:
+	TWeakObjectPtr<ULensComponent> RecordedComponent;
+};
diff --git a/Engine/Plugins/VirtualProduction/CameraCalibrationCore/Source/CameraCalibrationCoreMovieScene/Public/MovieSceneLensComponentTrack.h b/Engine/Plugins/VirtualProduction/CameraCalibrationCore/Source/CameraCalibrationCoreMovieScene/Public/MovieSceneLensComponentTrack.h
new file mode 100644
--- /dev/null
+++ b/Engine/Plugins/VirtualProduction/CameraCalibrationCore/Source/CameraCalibrationCoreMovieScene/Public/MovieSceneLensComponentTrack.h
@@ -0,0 +1,37 @@
+// Copyright Epic Games, Inc. All Rights Reserved.
+
+#pragma once
+
+#include "MovieSceneNameableTrack.h"
+
+#include "MovieSceneLensComponentTrack.generated.h"
+
+/** Movie Scene track for Lens Component */
+UCLASS()
+class CAMERACALIBRATIONCOREMOVIESCENE_API UMovieSceneLensComponentTrack : public UMovieSceneNameableTrack
+{
+	GENERATED_BODY()
+
+public:
+	//~ Begin UMovieSceneTrack interface
+	virtual void AddSection(UMovieSceneSection& Section) override;
+	virtual const TArray<UMovieSceneSection*>& GetAllSections() const override;
+	virtual bool HasSection(const UMovieSceneSection& Section) const override;
+	virtual bool IsEmpty() const override;
+	virtual void RemoveAllAnimationData() override;
+	virtual void RemoveSection(UMovieSceneSection& Section) override;
+	virtual void RemoveSectionAt(int32 SectionIndex) override;
+
+	virtual bool SupportsType(TSubclassOf<UMovieSceneSection> SectionClass) const override;
+	virtual UMovieSceneSection* CreateNewSection() override;
+
+#if WITH_EDITORONLY_DATA
+	virtual FText GetDisplayName() const override;
+#endif
+	//~ End UMovieSceneTrack interface
+
+protected:
+	/** Array of movie scene sections managed by this track */
+	UPROPERTY()
+	TArray<UMovieSceneSection*> Sections;
+};
\ No newline at end of file
diff --git a/Engine/Plugins/VirtualProduction/CameraCalibrationCore/Source/CameraCalibrationCoreSequencer/CameraCalibrationCoreSequencer.build.cs b/Engine/Plugins/VirtualProduction/CameraCalibrationCore/Source/CameraCalibrationCoreSequencer/CameraCalibrationCoreSequencer.build.cs
new file mode 100644
--- /dev/null
+++ b/Engine/Plugins/VirtualProduction/CameraCalibrationCore/Source/CameraCalibrationCoreSequencer/CameraCalibrationCoreSequencer.build.cs
@@ -0,0 +1,47 @@
+// Copyright Epic Games, Inc. All Rights Reserved.
+
+using UnrealBuildTool;
+
+public class CameraCalibrationCoreSequencer : ModuleRules
+{
+	public CameraCalibrationCoreSequencer(ReadOnlyTargetRules Target) : base(Target)
+	{
+		PCHUsage = ModuleRules.PCHUsageMode.UseExplicitOrSharedPCHs;
+
+		PublicIncludePaths.AddRange(
+			new string[] {
+			}
+			);
+				
+		
+		PrivateIncludePaths.AddRange(
+			new string[] {
+			}
+		);
+			
+		
+		PublicDependencyModuleNames.AddRange(
+			new string[]
+			{
+			}
+		);
+		
+		PrivateDependencyModuleNames.AddRange(
+			new string[]
+			{
+ 				"CameraCalibrationCore",
+ 				"CameraCalibrationCoreMovieScene",
+ 				"Core",
+ 				"CoreUObject",
+ 				"LevelSequence",
+ 				"MovieScene",
+ 				"MovieSceneTools",
+ 				"Sequencer",
+ 				"Slate",
+ 				"SlateCore",
+ 				"UnrealEd",
+ 				"TakeTrackRecorders",
+			}
+		);
+	}
+}
diff --git a/Engine/Plugins/VirtualProduction/CameraCalibrationCore/Source/CameraCalibrationCoreSequencer/Private/CameraCalibrationCoreSequencer.cpp b/Engine/Plugins/VirtualProduction/CameraCalibrationCore/Source/CameraCalibrationCoreSequencer/Private/CameraCalibrationCoreSequencer.cpp
new file mode 100644
--- /dev/null
+++ b/Engine/Plugins/VirtualProduction/CameraCalibrationCore/Source/CameraCalibrationCoreSequencer/Private/CameraCalibrationCoreSequencer.cpp
@@ -0,0 +1,46 @@
+// Copyright Epic Games, Inc. All Rights Reserved.
+
+#include "Modules/ModuleInterface.h"
+
+#include "Features/IModularFeatures.h"
+#include "ISequencerModule.h"
+#include "LensComponentTrackEditor.h"
+#include "Modules/ModuleManager.h"
+#include "MovieSceneLensComponentTrackRecorder.h"
+
+static const FName MovieSceneTrackRecorderFactoryName("MovieSceneTrackRecorderFactory");
+
+class FCameraCalibrationCoreSequencerModule : public IModuleInterface
+{
+public:
+	//~ Begin IModuleInterface interface
+	virtual void StartupModule() override
+	{
+		// Register Track Editor for LensComponent track
+		ISequencerModule& SequencerModule = FModuleManager::LoadModuleChecked<ISequencerModule>("Sequencer");
+		CreateLensComponentTrackEditorHandle = SequencerModule.RegisterTrackEditor(FOnCreateTrackEditor::CreateStatic(&FLensComponentTrackEditor::CreateTrackEditor));
+
+		// Register modular feature for LensComponent track recorder factory
+		IModularFeatures::Get().RegisterModularFeature(MovieSceneTrackRecorderFactoryName, &LensComponentTrackRecorderFactory);
+	}
+
+	virtual void ShutdownModule() override
+	{
+		// Unregister Track Editor for LensComponent track
+		if (ISequencerModule* SequencerModule = FModuleManager::GetModulePtr<ISequencerModule>("Sequencer"))
+		{
+			SequencerModule->UnRegisterTrackEditor(CreateLensComponentTrackEditorHandle);
+		}
+
+		// Unregister modular feature for LensComponent track recorder factory
+		IModularFeatures::Get().UnregisterModularFeature(MovieSceneTrackRecorderFactoryName, &LensComponentTrackRecorderFactory);
+	}
+	//~ End IModuleInterface interface
+
+private:
+	FDelegateHandle CreateLensComponentTrackEditorHandle;
+
+	FMovieSceneLensComponentTrackRecorderFactory LensComponentTrackRecorderFactory;
+};
+
+IMPLEMENT_MODULE(FCameraCalibrationCoreSequencerModule, CameraCalibrationCoreSequencer);
\ No newline at end of file
diff --git a/Engine/Plugins/VirtualProduction/CameraCalibrationCore/Source/CameraCalibrationCoreSequencer/Private/LensComponentTrackEditor.cpp b/Engine/Plugins/VirtualProduction/CameraCalibrationCore/Source/CameraCalibrationCoreSequencer/Private/LensComponentTrackEditor.cpp
new file mode 100644
--- /dev/null
+++ b/Engine/Plugins/VirtualProduction/CameraCalibrationCore/Source/CameraCalibrationCoreSequencer/Private/LensComponentTrackEditor.cpp
@@ -0,0 +1,25 @@
+// Copyright Epic Games, Inc. All Rights Reserved.
+
+#include "LensComponentTrackEditor.h"
+
+#include "LevelSequence.h"
+
+TSharedRef<ISequencerTrackEditor> FLensComponentTrackEditor::CreateTrackEditor(TSharedRef<ISequencer> InSequencer)
+{
+	return MakeShared<FLensComponentTrackEditor>(InSequencer);
+}
+
+bool FLensComponentTrackEditor::SupportsSequence(UMovieSceneSequence* InSequence) const
+{
+	if (!InSequence || InSequence->IsTrackSupported(UMovieSceneLensComponentTrack::StaticClass()) == ETrackSupport::NotSupported)
+	{
+		return false;
+	}
+
+	return InSequence->IsA(ULevelSequence::StaticClass());
+}
+
+bool FLensComponentTrackEditor::SupportsType(TSubclassOf<UMovieSceneTrack> Type) const
+{
+	return (Type == UMovieSceneLensComponentTrack::StaticClass());
+}
diff --git a/Engine/Plugins/VirtualProduction/CameraCalibrationCore/Source/CameraCalibrationCoreSequencer/Private/LensComponentTrackEditor.h b/Engine/Plugins/VirtualProduction/CameraCalibrationCore/Source/CameraCalibrationCoreSequencer/Private/LensComponentTrackEditor.h
new file mode 100644
--- /dev/null
+++ b/Engine/Plugins/VirtualProduction/CameraCalibrationCore/Source/CameraCalibrationCoreSequencer/Private/LensComponentTrackEditor.h
@@ -0,0 +1,30 @@
+// Copyright Epic Games, Inc. All Rights Reserved.
+
+#pragma once
+
+#include "KeyframeTrackEditor.h"
+#include "MovieSceneLensComponentTrack.h"
+
+/** A property track editor for Lens Component */
+class FLensComponentTrackEditor : public FKeyframeTrackEditor<UMovieSceneLensComponentTrack>
+{
+public:
+
+	FLensComponentTrackEditor(TSharedRef<ISequencer> InSequencer) : FKeyframeTrackEditor<UMovieSceneLensComponentTrack>(InSequencer)
+	{
+	}
+
+	/**
+	* Creates an instance of this class.  Called by a sequencer
+	*
+	* @param OwningSequencer The sequencer instance to be used by this tool
+	* @return The new instance of this class
+	*/
+	static TSharedRef<ISequencerTrackEditor> CreateTrackEditor(TSharedRef<ISequencer> OwningSequencer);
+
+public:
+	//~ Begin ISequencerTrackEditor interface
+	virtual bool SupportsSequence(UMovieSceneSequence* InSequence) const override;
+	virtual bool SupportsType(TSubclassOf<UMovieSceneTrack> Type) const override;
+	//~ End ISequencerTrackEditor interface
+};
diff --git a/Engine/Plugins/VirtualProduction/CameraCalibrationCore/Source/CameraCalibrationCoreSequencer/Private/MovieSceneLensComponentTrackRecorder.cpp b/Engine/Plugins/VirtualProduction/CameraCalibrationCore/Source/CameraCalibrationCoreSequencer/Private/MovieSceneLensComponentTrackRecorder.cpp
new file mode 100644
--- /dev/null
+++ b/Engine/Plugins/VirtualProduction/CameraCalibrationCore/Source/CameraCalibrationCoreSequencer/Private/MovieSceneLensComponentTrackRecorder.cpp
@@ -0,0 +1,82 @@
+// Copyright Epic Games, Inc. All Rights Reserved.
+
+#include "MovieSceneLensComponentTrackRecorder.h"
+
+#include "MovieSceneLensComponentTrack.h"
+#include "LensComponent.h"
+
+#define LOCTEXT_NAMESPACE "MovieSceneLensComponentTrackRecorder"
+
+bool FMovieSceneLensComponentTrackRecorderFactory::CanRecordObject(class UObject* InObjectToRecord) const
+{
+	return (InObjectToRecord->GetClass() == ULensComponent::StaticClass());
+}
+
+bool FMovieSceneLensComponentTrackRecorderFactory::CanRecordProperty(class UObject* InObjectToRecord, class FProperty* InPropertyToRecord) const
+{ 
+	return false; 
+}
+
+UMovieSceneTrackRecorder* FMovieSceneLensComponentTrackRecorderFactory::CreateTrackRecorderForObject() const
+{
+	return NewObject<UMovieSceneLensComponentTrackRecorder>();
+}
+
+UMovieSceneTrackRecorder* FMovieSceneLensComponentTrackRecorderFactory::CreateTrackRecorderForProperty(UObject* InObjectToRecord, const FName& InPropertyToRecord) const
+{
+	return nullptr; 
+}
+
+FText FMovieSceneLensComponentTrackRecorderFactory::GetDisplayName() const
+{ 
+	return LOCTEXT("LensComponentTrackRecorderDisplayName", "Lens Component"); 
+}
+
+void UMovieSceneLensComponentTrackRecorder::CreateTrackImpl()
+{
+	// Find/Add a Lens Component Track to the MovieScene
+	UMovieSceneLensComponentTrack* LensComponentTrack = MovieScene->FindTrack<UMovieSceneLensComponentTrack>(ObjectGuid);
+	if (!LensComponentTrack)
+	{
+		LensComponentTrack = MovieScene->AddTrack<UMovieSceneLensComponentTrack>(ObjectGuid);
+	}
+	else
+	{
+		LensComponentTrack->RemoveAllAnimationData();
+	}
+
+	// Create a new section for the Lens Component Track and initialize it with a pointer to the component being recorded
+	if (LensComponentTrack)
+	{
+		LensComponentSection = Cast<UMovieSceneLensComponentSection>(LensComponentTrack->CreateNewSection());
+
+		if (UMovieSceneLensComponentSection* Section = LensComponentSection.Get())
+		{
+			LensComponentTrack->AddSection(*Section);
+
+			ULensComponent* LensComponent = Cast<ULensComponent>(ObjectToRecord.Get());
+			Section->Initialize(LensComponent);
+		}
+	}
+}
+
+void UMovieSceneLensComponentTrackRecorder::RecordSampleImpl(const FQualifiedFrameTime& CurrentTime)
+{
+	if (UMovieSceneLensComponentSection* Section = LensComponentSection.Get())
+	{
+		FFrameRate TickResolution = Section->GetTypedOuter<UMovieScene>()->GetTickResolution();
+		FFrameNumber CurrentFrame = CurrentTime.ConvertTo(TickResolution).FloorToFrame();
+		Section->RecordFrame(CurrentFrame);
+		Section->SetEndFrame(CurrentFrame);
+	}
+}
+
+void UMovieSceneLensComponentTrackRecorder::FinalizeTrackImpl()
+{
+	if (UMovieSceneLensComponentSection* Section = LensComponentSection.Get())
+	{
+		Section->Finalize();
+	}
+}
+
+#undef LOCTEXT_NAMESPACE
diff --git a/Engine/Plugins/VirtualProduction/CameraCalibrationCore/Source/CameraCalibrationCoreSequencer/Private/MovieSceneLensComponentTrackRecorder.h b/Engine/Plugins/VirtualProduction/CameraCalibrationCore/Source/CameraCalibrationCoreSequencer/Private/MovieSceneLensComponentTrackRecorder.h
new file mode 100644
--- /dev/null
+++ b/Engine/Plugins/VirtualProduction/CameraCalibrationCore/Source/CameraCalibrationCoreSequencer/Private/MovieSceneLensComponentTrackRecorder.h
@@ -0,0 +1,43 @@
+// Copyright Epic Games, Inc. All Rights Reserved.
+
+#pragma once
+
+#include "TrackRecorders/IMovieSceneTrackRecorderFactory.h"
+#include "TrackRecorders/MovieSceneTrackRecorder.h"
+
+#include "MovieSceneLensComponentSection.h"
+
+#include "MovieSceneLensComponentTrackRecorder.generated.h"
+
+/** MovieScene Track Recorder Factory for a Lens Component */
+class FMovieSceneLensComponentTrackRecorderFactory : public IMovieSceneTrackRecorderFactory
+{
+public:
+	//~ Begin IMovieSceneTrackRecorderFactory interface
+	virtual bool CanRecordObject(class UObject* InObjectToRecord) const override;
+	virtual bool CanRecordProperty(class UObject* InObjectToRecord, class FProperty* InPropertyToRecord) const override;
+
+	virtual UMovieSceneTrackRecorder* CreateTrackRecorderForObject() const override;
+	virtual UMovieSceneTrackRecorder* CreateTrackRecorderForProperty(UObject* InObjectToRecord, const FName& InPropertyToRecord) const override;
+
+	virtual FText GetDisplayName() const override;
+	//~ End IMovieSceneTrackRecorderFactory interface
+};
+
+/** MovieScene Track Recorder for a Lens Component */
+UCLASS(BlueprintType)
+class UMovieSceneLensComponentTrackRecorder : public UMovieSceneTrackRecorder
+{
+	GENERATED_BODY()
+public:
+	//~ Begin UMovieSceneTrackRecorder Interface
+	virtual void CreateTrackImpl() override;
+	virtual void RecordSampleImpl(const FQualifiedFrameTime& CurrentTime) override;
+	virtual void FinalizeTrackImpl() override;
+	virtual UMovieSceneSection* GetMovieSceneSection() const override { return LensComponentSection.Get(); }
+	//~ End UMovieSceneTrackRecorder Interface
+
+private:
+	/** Section in which to record Lens Component data */
+	TWeakObjectPtr<UMovieSceneLensComponentSection> LensComponentSection;
+};
