[CameraCalibration] Fix uninitialized distortion displacement maps caused by materials not being compiled in time for the first draw

diff --git a/Engine/Plugins/VirtualProduction/CameraCalibrationCore/Source/CameraCalibrationCore/Private/LensDistortionModelHandlerBase.cpp b/Engine/Plugins/VirtualProduction/CameraCalibrationCore/Source/CameraCalibrationCore/Private/LensDistortionModelHandlerBase.cpp
--- a/Engine/Plugins/VirtualProduction/CameraCalibrationCore/Source/CameraCalibrationCore/Private/LensDistortionModelHandlerBase.cpp
+++ b/Engine/Plugins/VirtualProduction/CameraCalibrationCore/Source/CameraCalibrationCore/Private/LensDistortionModelHandlerBase.cpp
@@ -221,8 +221,47 @@ bool ULensDistortionModelHandlerBase::DrawDistortionDisplacementMap(UTextureRend
 	return true;
 }
 
+bool ULensDistortionModelHandlerBase::IsDisplacementMapMaterialReady(UMaterialInstanceDynamic* MID)
+{
+	ERHIFeatureLevel::Type FeatureLevel = GetWorld() ? (ERHIFeatureLevel::Type)GetWorld()->FeatureLevel : GMaxRHIFeatureLevel;
+	if (FMaterialResource* MaterialResource = MID->GetMaterialResource(FeatureLevel))
+	{
+		if (MaterialResource->IsGameThreadShaderMapComplete())
+		{
+			return true;
+		}
+
+		MaterialResource->SubmitCompileJobs_GameThread(EShaderCompileJobPriority::ForceLocal);
+
+		return MaterialResource->IsGameThreadShaderMapComplete();
+	}
+
+	return false;
+}
+
 void ULensDistortionModelHandlerBase::ProcessCurrentDistortion()
 {
+	bool bAreMaterialsReady = true;
+	if (!UndistortionDisplacementMapMID || !DistortionDisplacementMapMID)
+	{
+		InitDistortionMaterials();
+	}
+
+	if (UndistortionDisplacementMapMID && !IsDisplacementMapMaterialReady(UndistortionDisplacementMapMID))
+	{
+		bAreMaterialsReady = false;
+	}
+
+	if (DistortionDisplacementMapMID && !IsDisplacementMapMaterialReady(DistortionDisplacementMapMID))
+	{
+		bAreMaterialsReady = false;
+	}
+
+	if (!bAreMaterialsReady)
+	{
+		return;
+	}
+
 	if(bIsDirty)
 	{
 		bIsDirty = false;
diff --git a/Engine/Plugins/VirtualProduction/CameraCalibrationCore/Source/CameraCalibrationCore/Public/LensDistortionModelHandlerBase.h b/Engine/Plugins/VirtualProduction/CameraCalibrationCore/Source/CameraCalibrationCore/Public/LensDistortionModelHandlerBase.h
--- a/Engine/Plugins/VirtualProduction/CameraCalibrationCore/Source/CameraCalibrationCore/Public/LensDistortionModelHandlerBase.h
+++ b/Engine/Plugins/VirtualProduction/CameraCalibrationCore/Source/CameraCalibrationCore/Public/LensDistortionModelHandlerBase.h
@@ -134,6 +134,10 @@ protected:
 	/** Convert the generic distortion parameter array into the specific structure of parameters used by the supported lens model */
 	virtual void InterpretDistortionParameters() PURE_VIRTUAL(ULensDistortionModelHandlerBase::InterpretDistortionParameters);
 
+private:
+	/** Tests whether the input material's shaders have been compiled, indicating the material can be used for drawing */
+	bool IsDisplacementMapMaterialReady(UMaterialInstanceDynamic* MID);
+
 protected:
 	/** Lens Model describing how to interpret the distortion parameters */
 	UPROPERTY(VisibleAnywhere, BlueprintReadWrite, Category = "Distortion")
