Fix CIS issue with duplicate macro definition

diff --git a/Engine/Plugins/VirtualProduction/CameraCalibrationCore/Source/CameraCalibrationCore/Private/CameraCalibrationSubsystem.cpp b/Engine/Plugins/VirtualProduction/CameraCalibrationCore/Source/CameraCalibrationCore/Private/CameraCalibrationSubsystem.cpp
--- a/Engine/Plugins/VirtualProduction/CameraCalibrationCore/Source/CameraCalibrationCore/Private/CameraCalibrationSubsystem.cpp
+++ b/Engine/Plugins/VirtualProduction/CameraCalibrationCore/Source/CameraCalibrationCore/Private/CameraCalibrationSubsystem.cpp
@@ -89,9 +89,8 @@ ULensDistortionModelHandlerBase* UCameraCalibrationSubsystem::FindDistortionMode
 ULensDistortionModelHandlerBase* UCameraCalibrationSubsystem::FindOrCreateDistortionModelHandler(FDistortionHandlerPicker& DistortionHandlerPicker, const TSubclassOf<ULensModel> LensModelClass)
 {
 	// This function has been deprecated. The implementation has been changed to provide some backwards compatibility, but code should be updated to not call this function.
-PRAGMA_DISABLE_DEPRECATION_WARNINGS
 	TArray<ULensDistortionModelHandlerBase*> Handlers = GetDistortionModelHandlers(DistortionHandlerPicker.TargetCameraComponent);
-PRAGMA_ENABLE_DEPRECATION_WARNINGS
+
 	for (ULensDistortionModelHandlerBase* Handler : Handlers)
 	{
 		if (Handler->GetLensModelClass() == LensModelClass)
