[CameraCalibration] Fix focal length evaluation in LensComponent. Fixes bug where distortion could not be disabled properly on playback, and bug where override LensFile's focal length was not evaluated on playback.
Also, fix a bug where changing the distortion source to "Lens File" did not automatically update the lens model, leading to distortion not being properly applied.

diff --git a/Engine/Plugins/VirtualProduction/CameraCalibrationCore/Source/CameraCalibrationCore/Private/LensComponent.cpp b/Engine/Plugins/VirtualProduction/CameraCalibrationCore/Source/CameraCalibrationCore/Private/LensComponent.cpp
--- a/Engine/Plugins/VirtualProduction/CameraCalibrationCore/Source/CameraCalibrationCore/Private/LensComponent.cpp
+++ b/Engine/Plugins/VirtualProduction/CameraCalibrationCore/Source/CameraCalibrationCore/Private/LensComponent.cpp
@@ -82,7 +82,11 @@ void ULensComponent::TickComponent(float DeltaTime, ELevelTick TickType, FActorC
 		// If the camera's focal length has changed since the last time this component ticked, we need to update the original focal length we are using
 		if (CineCameraComponent->CurrentFocalLength != LastFocalLength)
 		{
-			OriginalFocalLength = CineCameraComponent->CurrentFocalLength;
+			// If this is a recorded lens component, then there are no external forces changing the focal length. Original focal length should be preserved.
+			if (EvaluationMode != EFIZEvaluationMode::UseRecordedValues)
+			{
+				OriginalFocalLength = CineCameraComponent->CurrentFocalLength;
+			}
 		}
 
 		// Get the focus and zoom values needed to evaluate the LensFile this Tick
@@ -98,10 +102,7 @@ void ULensComponent::TickComponent(float DeltaTime, ELevelTick TickType, FActorC
 
 		UpdateCameraFilmback(CineCameraComponent);
 
-		if (EvaluationMode != EFIZEvaluationMode::UseRecordedValues)
-		{
-			EvaluateFocalLength(CineCameraComponent);
-		}
+		EvaluateFocalLength(CineCameraComponent);
 
 		// Evaluate Distortion
 		bWasDistortionEvaluated = false;
@@ -181,19 +182,15 @@ void ULensComponent::TickComponent(float DeltaTime, ELevelTick TickType, FActorC
 				// Cache the latest distortion MID
 				LastDistortionMID = NewDistortionMID;
 
-				// If using recorded values, the overscanned FOV will already have been applied
-				if (EvaluationMode != EFIZEvaluationMode::UseRecordedValues)
-				{
-					// Get the overscan factor and use it to modify the target camera's FOV
-					const float OverscanFactor = LensDistortionHandler->GetOverscanFactor();
-					const float OverscanSensorWidth = CineCameraComponent->Filmback.SensorWidth * OverscanFactor;
-					const float OverscanFOV = FMath::RadiansToDegrees(2.0f * FMath::Atan(OverscanSensorWidth / (2.0f * OriginalFocalLength)));
-					CineCameraComponent->SetFieldOfView(OverscanFOV);
-
-					// Update the minimum and maximum focal length of the camera (if needed)
-					CineCameraComponent->LensSettings.MinFocalLength = FMath::Min(CineCameraComponent->LensSettings.MinFocalLength, CineCameraComponent->CurrentFocalLength);
-					CineCameraComponent->LensSettings.MaxFocalLength = FMath::Max(CineCameraComponent->LensSettings.MaxFocalLength, CineCameraComponent->CurrentFocalLength);
-				}
+				// Get the overscan factor and use it to modify the target camera's FOV
+				const float OverscanFactor = LensDistortionHandler->GetOverscanFactor();
+				const float OverscanSensorWidth = CineCameraComponent->Filmback.SensorWidth * OverscanFactor;
+				const float OverscanFOV = FMath::RadiansToDegrees(2.0f * FMath::Atan(OverscanSensorWidth / (2.0f * OriginalFocalLength)));
+				CineCameraComponent->SetFieldOfView(OverscanFOV);
+
+				// Update the minimum and maximum focal length of the camera (if needed)
+				CineCameraComponent->LensSettings.MinFocalLength = FMath::Min(CineCameraComponent->LensSettings.MinFocalLength, CineCameraComponent->CurrentFocalLength);
+				CineCameraComponent->LensSettings.MaxFocalLength = FMath::Max(CineCameraComponent->LensSettings.MaxFocalLength, CineCameraComponent->CurrentFocalLength);
 
 				bIsDistortionSetup = true;
 			}
@@ -202,6 +199,14 @@ void ULensComponent::TickComponent(float DeltaTime, ELevelTick TickType, FActorC
 				CleanupDistortion(CineCameraComponent);
 			}
 		}
+		else
+		{
+			// If this is a recorded lens component, and distortion is disabled, set the focal length the original focal length to override the recorded value from the camera.
+			if (EvaluationMode == EFIZEvaluationMode::UseRecordedValues)
+			{
+				CineCameraComponent->CurrentFocalLength = OriginalFocalLength;
+			}
+		}
 
 		LastFocalLength = CineCameraComponent->CurrentFocalLength;
 	}
@@ -303,6 +308,10 @@ void ULensComponent::PostEditChangeProperty(struct FPropertyChangedEvent& Proper
 	{
 		SetLensModel(LensModel);
 	}
+	else if (PropertyName == GET_MEMBER_NAME_CHECKED(ULensComponent, DistortionStateSource))
+	{
+		SetDistortionSource(DistortionStateSource);
+	}
 
 	Super::PostEditChangeProperty(PropertyChangedEvent);
 }
@@ -598,6 +607,15 @@ EDistortionSource ULensComponent::GetDistortionSource() const
 void ULensComponent::SetDistortionSource(EDistortionSource Source)
 {
 	DistortionStateSource = Source;
+
+	// If the new source is a LensFile, update the lens model to match
+	if (DistortionStateSource == EDistortionSource::LensFile)
+	{
+		if (ULensFile* LensFile = LensFilePicker.GetLensFile())
+		{
+			LensModel = LensFile->LensInfo.LensModel;
+		}
+	}
 }
 
 bool ULensComponent::ShouldApplyDistortion() const
