CameraCalibrationCore automation test category changed to "Plugin"

diff --git a/Engine/Plugins/VirtualProduction/CameraCalibrationCore/Source/CameraCalibrationCore/Private/Tests/TestDistortionMapping.cpp b/Engine/Plugins/VirtualProduction/CameraCalibrationCore/Source/CameraCalibrationCore/Private/Tests/TestDistortionMapping.cpp
--- a/Engine/Plugins/VirtualProduction/CameraCalibrationCore/Source/CameraCalibrationCore/Private/Tests/TestDistortionMapping.cpp
+++ b/Engine/Plugins/VirtualProduction/CameraCalibrationCore/Source/CameraCalibrationCore/Private/Tests/TestDistortionMapping.cpp
@@ -13,7 +13,7 @@
 #include "SphericalLensDistortionModelHandler.h"
 #include "UObject/StrongObjectPtr.h"
 
-IMPLEMENT_SIMPLE_AUTOMATION_TEST(FTestCameraCalibrationCore, "Plugins.CameraCalibrationCore", EAutomationTestFlags::EditorContext | EAutomationTestFlags::EngineFilter)
+IMPLEMENT_SIMPLE_AUTOMATION_TEST(FTestCameraCalibrationCore, "Plugin.CameraCalibrationCore", EAutomationTestFlags::EditorContext | EAutomationTestFlags::EngineFilter)
 
 
 namespace CameraCalibrationTestUtil
