[CameraCalibration] Fix lens component to properly detect livelink roles that derive from the camera role

diff --git a/Engine/Plugins/VirtualProduction/CameraCalibrationCore/Source/CameraCalibrationCore/Private/LensComponent.cpp b/Engine/Plugins/VirtualProduction/CameraCalibrationCore/Source/CameraCalibrationCore/Private/LensComponent.cpp
--- a/Engine/Plugins/VirtualProduction/CameraCalibrationCore/Source/CameraCalibrationCore/Private/LensComponent.cpp
+++ b/Engine/Plugins/VirtualProduction/CameraCalibrationCore/Source/CameraCalibrationCore/Private/LensComponent.cpp
@@ -715,15 +715,16 @@ void ULensComponent::OnLiveLinkComponentRegistered(ULiveLinkComponentController*
 void ULensComponent::ProcessLiveLinkData(const ULiveLinkComponentController* const LiveLinkComponent, const FLiveLinkSubjectFrameData& SubjectData)
 {
 	TSubclassOf<ULiveLinkRole> LiveLinkRole = LiveLinkComponent->GetSubjectRepresentation().Role;
-	if (LiveLinkRole == ULiveLinkCameraRole::StaticClass())
+
+	if (LiveLinkRole.Get()->IsChildOf(ULiveLinkCameraRole::StaticClass()))
 	{
-		UpdateTrackedComponent(LiveLinkComponent, SubjectData);
 		UpdateLiveLinkFIZ(LiveLinkComponent, SubjectData);
 	}
-	else if (LiveLinkRole == ULiveLinkTransformRole::StaticClass())
-	{
-		UpdateTrackedComponent(LiveLinkComponent, SubjectData);
-	}
+
+	if (LiveLinkRole.Get()->IsChildOf(ULiveLinkTransformRole::StaticClass()))
+ 	{
+ 		UpdateTrackedComponent(LiveLinkComponent, SubjectData);
+ 	}
 }
 
 void ULensComponent::UpdateTrackedComponent(const ULiveLinkComponentController* const LiveLinkComponent, const FLiveLinkSubjectFrameData& SubjectData)
