Fix ST maps importing incorrectly due to incorrect assumptions about channels used and the image origin.

diff --git a/Engine/Plugins/VirtualProduction/CameraCalibration/Source/CameraCalibrationEditor/Private/AssetEditor/SLensDataAddPointDialog.cpp b/Engine/Plugins/VirtualProduction/CameraCalibration/Source/CameraCalibrationEditor/Private/AssetEditor/SLensDataAddPointDialog.cpp
--- a/Engine/Plugins/VirtualProduction/CameraCalibration/Source/CameraCalibrationEditor/Private/AssetEditor/SLensDataAddPointDialog.cpp
+++ b/Engine/Plugins/VirtualProduction/CameraCalibration/Source/CameraCalibrationEditor/Private/AssetEditor/SLensDataAddPointDialog.cpp
@@ -523,7 +523,44 @@ TSharedRef<SWidget> SLensDataAddPointDialog::MakeLensDataWidget()
 		{
 			STMapData->InitializeAs<FSTMapInfo>();
 			TSharedPtr<IStructureDetailsView> StructureDetailsView = PropertyEditor.CreateStructureDetailView(DetailArgs, StructureViewArgs, STMapData);
-			LensDataWidget = StructureDetailsView->GetWidget();
+
+			FocalLengthData->InitializeAs<FFocalLengthInfo>();
+			TSharedPtr<IStructureDetailsView> FocalLengthStructDetailsView = PropertyEditor.CreateStructureDetailView(DetailArgs, StructureViewArgs, FocalLengthData);
+			FFocalLengthInfo* FocalLengthInstanceData = FocalLengthData->Cast<FFocalLengthInfo>();
+			FocalLengthInstanceData->FxFy = GetDefaultFocalLengthValue();
+
+			LensDataWidget =
+				SNew(SVerticalBox)
+				+ SVerticalBox::Slot()
+				.AutoHeight()
+				.Padding(5.0f, 5.0f)
+				[
+					SNew(STextBlock)
+					.Text(LOCTEXT("FxFyDisplayUnitHelpText",
+						"Choose the units that should be used for the value of FxFy below.\n"
+						"If using mm or pixels, check that the sensor and image dimensions are set correctly in the Lens Information."
+					))
+					.AutoWrapText(true)
+				]
+				+ SVerticalBox::Slot()
+				.AutoHeight()
+				.Padding(0.0f, 5.0f)
+				[
+					UnitWidget.ToSharedRef()
+				]
+				+ SVerticalBox::Slot()
+				.AutoHeight()
+				.Padding(0.0f, 5.0f)
+				[
+					StructureDetailsView->GetWidget().ToSharedRef()
+				]
+				+ SVerticalBox::Slot()
+				.AutoHeight()
+				.Padding(0.0f, 5.0f)
+				[
+					FocalLengthStructDetailsView->GetWidget().ToSharedRef()
+				];
+
 			break;
 		}
 		default:
@@ -818,6 +855,10 @@ void SLensDataAddPointDialog::AddDataToLensFile() const
 		case ELensDataCategory::STMap:
 		{
 			LensFile->AddSTMapPoint(TrackingInputData[0].Value, TrackingInputData[1].Value, *STMapData->Get());
+
+			FFocalLengthInfo FocalLengthValue = *FocalLengthData->Get();
+			NormalizeValue(FocalLengthValue.FxFy);
+			LensFile->AddFocalLengthPoint(TrackingInputData[0].Value, TrackingInputData[1].Value, FocalLengthValue);
 			break;
 		}
 		default:
diff --git a/Engine/Plugins/VirtualProduction/CameraCalibrationCore/Shaders/Private/DistortionSTMapProcessor.usf b/Engine/Plugins/VirtualProduction/CameraCalibrationCore/Shaders/Private/DistortionSTMapProcessor.usf
--- a/Engine/Plugins/VirtualProduction/CameraCalibrationCore/Shaders/Private/DistortionSTMapProcessor.usf
+++ b/Engine/Plugins/VirtualProduction/CameraCalibrationCore/Shaders/Private/DistortionSTMapProcessor.usf
@@ -15,50 +15,82 @@ RWBuffer<float2> OutDistortedUV;
 RWTexture2D<float2> OutUndistortionDisplacementMap;
 RWTexture2D<float2> OutDistortionDisplacementMap;
 
+// Must match ECalibratedMapPixelOrigin
+#define PIXEL_ORIGIN_TOPLEFT 0
+#define PIXEL_ORIGIN_BOTTOMLEFT 1
+
+// Must match ECalibratedMapChannels
+#define DISTORTION_CHANNELS_RG 0
+#define DISTORTION_CHANNELS_BA 1
+#define DISTORTION_CHANNELS_NONE 2
+
+// Must match ECalibratedMapChannels
+#define UNDISTORTION_CHANNELS_RG 0
+#define UNDISTORTION_CHANNELS_BA 1
+#define UNDISTORTION_CHANNELS_NONE 2
 
 [numthreads(8,8,1)]
 void DistortionSTMapProcessorMainCS(uint3 ThreadId : SV_DispatchThreadID)
 {
 	float2 UV = TexelSize * (ThreadId.xy + 0.5f);
+
+#if PIXEL_ORIGIN == PIXEL_ORIGIN_TOPLEFT
 	float4 UVMapping = DistortionSTMap.SampleLevel(DistortionSTMapSampler, UV, 0);
+#elif PIXEL_ORIGIN == PIXEL_ORIGIN_BOTTOMLEFT
+	float2 MirroredUV = float2(UV.x, 1.0f - UV.y);
+	float4 UVMapping = DistortionSTMap.SampleLevel(DistortionSTMapSampler, MirroredUV, 0);
+#endif
+
+#if UNDISTORTION_CHANNELS == UNDISTORTION_CHANNELS_RG
+	float2 UndistortedPixel = UVMapping.rg;
+#elif UNDISTORTION_CHANNELS == UNDISTORTION_CHANNELS_BA
+	float2 UndistortedPixel = UVMapping.ba;
+#elif UNDISTORTION_CHANNELS == UNDISTORTION_CHANNELS_NONE
+	float2 UndistortedPixel = UV;
+#endif
 
-	float2 Undistortion = float2(UVMapping.x - UV.x, (1.0f - UVMapping.y) - UV.y);
-	float2 Distortion = float2(UVMapping.z - UV.x, (1.0f - UVMapping.w) - UV.y);
+#if DISTORTION_CHANNELS == DISTORTION_CHANNELS_RG
+	float2 DistortedPixel = UVMapping.rg;
+#elif DISTORTION_CHANNELS == DISTORTION_CHANNELS_BA
+	float2 DistortedPixel = UVMapping.ba;
+#elif DISTORTION_CHANNELS == DISTORTION_CHANNELS_NONE
+	float2 DistortedPixel = UV;
+#endif
 
-	OutUndistortionDisplacementMap[ThreadId.xy] = Undistortion;
-	OutDistortionDisplacementMap[ThreadId.xy] = Distortion;
+	OutUndistortionDisplacementMap[ThreadId.xy] = UndistortedPixel - UV;
+	OutDistortionDisplacementMap[ThreadId.xy] = DistortedPixel - UV;
 
 	//Stamp distorted UV for known undistorted points
 	if(all(ThreadId.xy == uint2(0,0))) //0.0f, 0.0f
 	{
-		OutDistortedUV[0] = UVMapping.zw;
+		OutDistortedUV[0] = DistortedPixel;
 	}
 	else if(all(ThreadId.xy == uint2(TextureSize.x * 0.5f, 0))) //0.5f, 0.0f
 	{
-		OutDistortedUV[1] = UVMapping.zw;
+		OutDistortedUV[1] = DistortedPixel;
 	}
 	else if(all(ThreadId.xy == uint2(TextureSize.x - 1, 0))) //1.0f, 0.0f
 	{
-		OutDistortedUV[2] = UVMapping.zw;
+		OutDistortedUV[2] = DistortedPixel;
 	}
 	else if(all(ThreadId.xy == uint2(TextureSize.x - 1, TextureSize.y * 0.5f))) //1.0f, 0.5f
 	{
-		OutDistortedUV[3] = UVMapping.zw;
+		OutDistortedUV[3] = DistortedPixel;
 	}
 	else if(all(ThreadId.xy == uint2(TextureSize.x - 1, TextureSize.y - 1))) //1.0f, 1.0f
 	{
-		OutDistortedUV[4] = UVMapping.zw;
+		OutDistortedUV[4] = DistortedPixel;
 	}
 	else if(all(ThreadId.xy == uint2(TextureSize.x * 0.5f, TextureSize.y - 1))) //0.5f, 1.0f
 	{
-		OutDistortedUV[5] = UVMapping.zw;
+		OutDistortedUV[5] = DistortedPixel;
 	}
 	else if(all(ThreadId.xy == uint2(0, TextureSize.y - 1))) //0.0f, 1.0f
 	{
-		OutDistortedUV[6] = UVMapping.zw;
+		OutDistortedUV[6] = DistortedPixel;
 	}
 	else if(all(ThreadId.xy == uint2(0, TextureSize.y * 0.5f))) //0.0f, 0.5f
 	{
-		OutDistortedUV[7] = UVMapping.zw;
+		OutDistortedUV[7] = DistortedPixel;
 	}	
 }
diff --git a/Engine/Plugins/VirtualProduction/CameraCalibrationCore/Source/CameraCalibrationCore/Private/CalibratedMapProcessor.cpp b/Engine/Plugins/VirtualProduction/CameraCalibrationCore/Source/CameraCalibrationCore/Private/CalibratedMapProcessor.cpp
--- a/Engine/Plugins/VirtualProduction/CameraCalibrationCore/Source/CameraCalibrationCore/Private/CalibratedMapProcessor.cpp
+++ b/Engine/Plugins/VirtualProduction/CameraCalibrationCore/Source/CameraCalibrationCore/Private/CalibratedMapProcessor.cpp
@@ -2,15 +2,14 @@
 
 #include "CalibratedMapProcessor.h"
 
+#include "CameraCalibrationCoreLog.h"
 #include "Engine/Texture.h"
 #include "Engine/TextureRenderTarget2D.h"
 #include "GlobalShader.h"		
-#include "CameraCalibrationCoreLog.h"
 #include "PixelShaderUtils.h"
 #include "RenderGraphUtils.h"
-#include "ShaderParameterStruct.h"
 #include "ScreenPass.h"
-
+#include "ShaderParameterStruct.h"
 
 
 class FCalibratedMapDerivedDataCS : public FGlobalShader
@@ -18,6 +17,12 @@ class FCalibratedMapDerivedDataCS : public FGlobalShader
 public:
 	DECLARE_GLOBAL_SHADER(FCalibratedMapDerivedDataCS);
 
+	class FMapPixelOrigin : SHADER_PERMUTATION_ENUM_CLASS("PIXEL_ORIGIN", ECalibratedMapPixelOrigin);
+	class FUndistortionChannels : SHADER_PERMUTATION_ENUM_CLASS("UNDISTORTION_CHANNELS", ECalibratedMapChannels);
+	class FDistortionChannels : SHADER_PERMUTATION_ENUM_CLASS("DISTORTION_CHANNELS", ECalibratedMapChannels);
+
+	using FPermutationDomain = TShaderPermutationDomain<FMapPixelOrigin, FUndistortionChannels, FDistortionChannels>;
+
 	SHADER_USE_PARAMETER_STRUCT(FCalibratedMapDerivedDataCS, FGlobalShader);
 
 	BEGIN_SHADER_PARAMETER_STRUCT(FParameters, )
@@ -27,7 +32,7 @@ public:
 		SHADER_PARAMETER_RDG_TEXTURE(Texture2D, DistortionSTMap)
 		SHADER_PARAMETER_SAMPLER(SamplerState, DistortionSTMapSampler)
 
-		SHADER_PARAMETER_RDG_BUFFER_UAV(RWBuffer<FVector2D>, OutDistortedUV)
+		SHADER_PARAMETER_RDG_BUFFER_UAV(RWBuffer<FVector2f>, OutDistortedUV)
 		SHADER_PARAMETER_RDG_TEXTURE_UAV(RWTexture2D<float2>, OutUndistortionDisplacementMap)
 		SHADER_PARAMETER_RDG_TEXTURE_UAV(RWTexture2D<float2>, OutDistortionDisplacementMap)
 		END_SHADER_PARAMETER_STRUCT()
@@ -159,7 +164,7 @@ void FCalibratedMapProcessor::ExecuteJob(TSharedPtr<FDerivedDistortionDataJob> J
 			CompletedJobs.Enqueue(Job);
 			return;
 		}
-		
+	
 		FTextureResource* SourceDistortionMap = SourceTexture->GetResource();
 		FTextureResource* DestinationUndistortionDisplacementMap = DestinationUndistortionTexture->GetResource();
 		FTextureResource* DestinationDistortionDisplacementMap = DestinationDistortionTexture->GetResource();
@@ -195,8 +200,13 @@ void FCalibratedMapProcessor::ExecuteJob(TSharedPtr<FDerivedDistortionDataJob> J
 			Parameters->OutUndistortionDisplacementMap = GraphBuilder.CreateUAV(UndistortionDisplacementMap);
 			Parameters->OutDistortionDisplacementMap = GraphBuilder.CreateUAV(DistortionDisplacementMap);
 
+			FCalibratedMapDerivedDataCS::FPermutationDomain PermutationVector;
+			PermutationVector.Set<FCalibratedMapDerivedDataCS::FMapPixelOrigin>(Job->JobArgs.Format.PixelOrigin);
+			PermutationVector.Set<FCalibratedMapDerivedDataCS::FUndistortionChannels>(Job->JobArgs.Format.UndistortionChannels);
+			PermutationVector.Set<FCalibratedMapDerivedDataCS::FDistortionChannels>(Job->JobArgs.Format.DistortionChannels);
+
             FGlobalShaderMap* GlobalShaderMap = GetGlobalShaderMap(GMaxRHIFeatureLevel);
-			TShaderMapRef<FCalibratedMapDerivedDataCS> ComputeShader(GlobalShaderMap);
+			TShaderMapRef<FCalibratedMapDerivedDataCS> ComputeShader(GlobalShaderMap, PermutationVector);
 
             GraphBuilder.AddPass(
             			RDG_EVENT_NAME("LensFileSTMapConversion"),
@@ -247,13 +257,14 @@ void FCalibratedMapProcessor::Update_RenderThread()
 	{
 		constexpr int32 EntriesCount = FDerivedDistortionDataJobOutput::EdgePointCount;
 		const int32 ReadbackDataByteCount = CompletedJob->Output.EdgePointsDistortedUVs.GetTypeSize();
-		const FVector2D* EdgePointsDistortedUV = static_cast<const FVector2D*>(CompletedJob->Readback->Lock(EntriesCount * ReadbackDataByteCount));
+		const FVector2f* EdgePointsDistortedUV = static_cast<const FVector2f*>(CompletedJob->Readback->Lock(EntriesCount * ReadbackDataByteCount));
 
 		if (EdgePointsDistortedUV)
 		{
 			for (int32 Index = 0; Index < CompletedJob->Output.EdgePointsDistortedUVs.Num(); ++Index)
 			{
-				CompletedJob->Output.EdgePointsDistortedUVs[Index] = EdgePointsDistortedUV[Index];
+				const FVector2f DistortedUV = EdgePointsDistortedUV[Index];
+				CompletedJob->Output.EdgePointsDistortedUVs[Index] = FVector2D(DistortedUV.X, DistortedUV.Y);
 				CompletedJob->Output.Result = EDerivedDistortionDataResult::Success;
 			}
 		}
diff --git a/Engine/Plugins/VirtualProduction/CameraCalibrationCore/Source/CameraCalibrationCore/Private/CalibratedMapProcessor.h b/Engine/Plugins/VirtualProduction/CameraCalibrationCore/Source/CameraCalibrationCore/Private/CalibratedMapProcessor.h
--- a/Engine/Plugins/VirtualProduction/CameraCalibrationCore/Source/CameraCalibrationCore/Private/CalibratedMapProcessor.h
+++ b/Engine/Plugins/VirtualProduction/CameraCalibrationCore/Source/CameraCalibrationCore/Private/CalibratedMapProcessor.h
@@ -3,6 +3,8 @@
 #pragma once
 
 #include "ICalibratedMapProcessor.h"
+
+#include "CalibratedMapFormat.h"
 #include "Containers/Queue.h"
 #include "Templates/PimplPtr.h"
 #include "UObject/WeakObjectPtr.h"
@@ -32,6 +34,9 @@ struct FDerivedDistortionDataJobArgs
 	/** Zoom point this is associated with */
 	float Zoom = 0.0f;
 
+	/** Formating information, such as pixel origin and the channels used for each displacement map */
+	FCalibratedMapFormat Format;
+
 	/** Source calibrated map */
 	//When coefficient route is supported, we might look at how we can have a mutual interface
 	TWeakObjectPtr<UTexture> SourceDistortionMap;
diff --git a/Engine/Plugins/VirtualProduction/CameraCalibrationCore/Source/CameraCalibrationCore/Private/LensFile.cpp b/Engine/Plugins/VirtualProduction/CameraCalibrationCore/Source/CameraCalibrationCore/Private/LensFile.cpp
--- a/Engine/Plugins/VirtualProduction/CameraCalibrationCore/Source/CameraCalibrationCore/Private/LensFile.cpp
+++ b/Engine/Plugins/VirtualProduction/CameraCalibrationCore/Source/CameraCalibrationCore/Private/LensFile.cpp
@@ -1347,13 +1347,20 @@ void ULensFile::UpdateDerivedData()
 			{
 				if (ZoomPoint.DerivedDistortionData.bIsDirty)
 				{
-					//Early exit if source data is invalid
+					//Early exit if source map does not exist
 					if (ZoomPoint.STMapInfo.DistortionMap == nullptr)
 					{
 						ZoomPoint.DerivedDistortionData.bIsDirty = false;
 						continue;
 					}
 
+					//Early exit it the source map is not yet loaded (but leave it marked dirty so it tries again later)
+					if (ZoomPoint.STMapInfo.DistortionMap->GetResource() == nullptr ||
+						ZoomPoint.STMapInfo.DistortionMap->GetResource()->IsProxy())
+					{
+						continue;
+					}
+
 					const FIntPoint CurrentDisplacementMapResolution = GetDefault<UCameraCalibrationSettings>()->GetDisplacementMapResolution();
 
 					//Create required undistortion texture for newly added points
@@ -1379,6 +1386,7 @@ void ULensFile::UpdateDerivedData()
 					FDerivedDistortionDataJobArgs JobArgs;
 					JobArgs.Focus = FocusPoint.Focus;
 					JobArgs.Zoom = ZoomPoint.Zoom;
+					JobArgs.Format = ZoomPoint.STMapInfo.MapFormat;
 					JobArgs.SourceDistortionMap = ZoomPoint.STMapInfo.DistortionMap;
 					JobArgs.OutputUndistortionDisplacementMap = ZoomPoint.DerivedDistortionData.UndistortionDisplacementMap;
 					JobArgs.OutputDistortionDisplacementMap = ZoomPoint.DerivedDistortionData.DistortionDisplacementMap;
diff --git a/Engine/Plugins/VirtualProduction/CameraCalibrationCore/Source/CameraCalibrationCore/Public/CalibratedMapFormat.h b/Engine/Plugins/VirtualProduction/CameraCalibrationCore/Source/CameraCalibrationCore/Public/CalibratedMapFormat.h
new file mode 100644
--- /dev/null
+++ b/Engine/Plugins/VirtualProduction/CameraCalibrationCore/Source/CameraCalibrationCore/Public/CalibratedMapFormat.h
@@ -0,0 +1,48 @@
+// Copyright Epic Games, Inc. All Rights Reserved.
+
+#pragma once
+
+#include "CoreMinimal.h"
+
+#include "CalibratedMapFormat.generated.h"
+
+/** Specifies the location of the pixel origin in a calibrated map */
+UENUM()
+enum class ECalibratedMapPixelOrigin : uint8
+{
+	TopLeft,
+	BottomLeft,
+
+	MAX // Required option for shader permutation
+};
+
+/** Specifies which two channels contain the calibrated map data (or None if there is no data) */
+UENUM()
+enum class ECalibratedMapChannels : uint8
+{
+	RG,
+	BA,
+	None,
+
+	MAX // Required option for shader permutation
+};
+
+/** Formatting options for processing a calibrated map */
+USTRUCT()
+struct CAMERACALIBRATIONCORE_API FCalibratedMapFormat
+{
+	GENERATED_BODY()
+
+public:
+	/** Specifies where in the image the (0, 0) pixel is */
+	UPROPERTY(EditAnywhere, Category="Format")
+	ECalibratedMapPixelOrigin PixelOrigin = ECalibratedMapPixelOrigin::TopLeft;
+
+	/** Specifies which two channels contain the undistortion map (or None if there is no undistortion data) */
+	UPROPERTY(EditAnywhere, Category = "Format")
+	ECalibratedMapChannels UndistortionChannels = ECalibratedMapChannels::RG;
+
+	/** Specifies which two channels contain the distortion map (or None if there is no distortion data) */
+	UPROPERTY(EditAnywhere, Category = "Format")
+	ECalibratedMapChannels DistortionChannels = ECalibratedMapChannels::BA;
+};
\ No newline at end of file
diff --git a/Engine/Plugins/VirtualProduction/CameraCalibrationCore/Source/CameraCalibrationCore/Public/LensData.h b/Engine/Plugins/VirtualProduction/CameraCalibrationCore/Source/CameraCalibrationCore/Public/LensData.h
--- a/Engine/Plugins/VirtualProduction/CameraCalibrationCore/Source/CameraCalibrationCore/Public/LensData.h
+++ b/Engine/Plugins/VirtualProduction/CameraCalibrationCore/Source/CameraCalibrationCore/Public/LensData.h
@@ -4,6 +4,7 @@
 
 #include "CoreTypes.h"
 
+#include "CalibratedMapFormat.h"
 #include "Engine/EngineTypes.h"
 #include "Engine/Texture.h"
 #include "Models/LensModel.h"
@@ -95,13 +96,13 @@ struct CAMERACALIBRATIONCORE_API FSTMapInfo
 	GENERATED_BODY()
 
 public:
-	/** 
-	 * Pre calibrated UVMap/STMap
-	 * RG channels are expected to have undistortion map (from distorted to undistorted)
-	 * BA channels are expected to have distortion map (from undistorted (CG) to distorted)
-	 */
+	/** Pre calibrated UVMap/STMap */
 	UPROPERTY(EditAnywhere, BlueprintReadWrite, Category = "Distortion")
 	TObjectPtr<UTexture> DistortionMap = nullptr;
+
+	/** Calibrated map format */
+	UPROPERTY(EditAnywhere, Category = "Format", meta = (ShowOnlyInnerProperties))
+	FCalibratedMapFormat MapFormat;
 };
 
 /**
