[CameraCalibration] Fix floating-point precision error when comparing camera filmback in LensFile

diff --git a/Engine/Plugins/VirtualProduction/CameraCalibrationCore/Source/CameraCalibrationCore/Private/LensFile.cpp b/Engine/Plugins/VirtualProduction/CameraCalibrationCore/Source/CameraCalibrationCore/Private/LensFile.cpp
--- a/Engine/Plugins/VirtualProduction/CameraCalibrationCore/Source/CameraCalibrationCore/Private/LensFile.cpp
+++ b/Engine/Plugins/VirtualProduction/CameraCalibrationCore/Source/CameraCalibrationCore/Private/LensFile.cpp
@@ -711,7 +711,7 @@ bool ULensFile::EvaluateDistortionForSTMaps(float InFocus, float InZoom, FVector
 		return true;
 	}
 
-	if ((LensInfo.SensorDimensions.X < InFilmback.X) || (LensInfo.SensorDimensions.Y < InFilmback.Y))
+	if (((LensInfo.SensorDimensions.X + UE_DOUBLE_KINDA_SMALL_NUMBER) < InFilmback.X) || ((LensInfo.SensorDimensions.Y + UE_DOUBLE_KINDA_SMALL_NUMBER) < InFilmback.Y))
 	{
 		UE_LOG(LogCameraCalibrationCore, Verbose
 			, TEXT("Can't evaluate LensFile '%s' - The filmback used to generate the calibrated ST Maps is smaller than" \
