[CameraCalibration] Fix camera distortion to properly reset when LiveLink Lens subject is destroyed.

diff --git a/Engine/Plugins/VirtualProduction/CameraCalibrationCore/Source/CameraCalibrationCore/Private/LensComponent.cpp b/Engine/Plugins/VirtualProduction/CameraCalibrationCore/Source/CameraCalibrationCore/Private/LensComponent.cpp
--- a/Engine/Plugins/VirtualProduction/CameraCalibrationCore/Source/CameraCalibrationCore/Private/LensComponent.cpp
+++ b/Engine/Plugins/VirtualProduction/CameraCalibrationCore/Source/CameraCalibrationCore/Private/LensComponent.cpp
@@ -478,6 +478,10 @@ void ULensComponent::SetLensFilePicker(FLensFilePicker LensFile)
 		{
 			SetLensModel(NewLensFile->LensInfo.LensModel);
 		}
+		else
+		{
+			ClearDistortionState();
+		}
 	}
 }
 
@@ -493,6 +497,10 @@ void ULensComponent::SetLensFile(ULensFile* Lens)
 		{
 			SetLensModel(Lens->LensInfo.LensModel);
 		}
+		else
+		{
+			ClearDistortionState();
+		}
 	}
 }
 
@@ -608,6 +616,8 @@ void ULensComponent::SetDistortionSource(EDistortionSource Source)
 {
 	DistortionStateSource = Source;
 
+	ClearDistortionState();
+
 	// If the new source is a LensFile, update the lens model to match
 	if (DistortionStateSource == EDistortionSource::LensFile)
 	{
@@ -683,6 +693,17 @@ void ULensComponent::CleanupDistortion(UCineCameraComponent* const CineCameraCom
 	bIsDistortionSetup = false;
 }
 
+void ULensComponent::ClearDistortionState()
+{
+	FLensDistortionState ZeroDistortionState;
+	if (LensModel)
+	{
+		const uint32 NumDistortionParameters = LensModel->GetDefaultObject<ULensModel>()->GetNumParameters();
+		ZeroDistortionState.DistortionInfo.Parameters.SetNumZeroed(NumDistortionParameters);
+	}
+	DistortionState = ZeroDistortionState;
+}
+
 void ULensComponent::CreateDistortionHandler()
 {
 	if (LensModel)
diff --git a/Engine/Plugins/VirtualProduction/CameraCalibrationCore/Source/CameraCalibrationCore/Public/LensComponent.h b/Engine/Plugins/VirtualProduction/CameraCalibrationCore/Source/CameraCalibrationCore/Public/LensComponent.h
--- a/Engine/Plugins/VirtualProduction/CameraCalibrationCore/Source/CameraCalibrationCore/Public/LensComponent.h
+++ b/Engine/Plugins/VirtualProduction/CameraCalibrationCore/Source/CameraCalibrationCore/Public/LensComponent.h
@@ -164,6 +164,10 @@ public:
 	UFUNCTION(BlueprintCallable, Category = "Lens Component")
 	void SetDistortionState(FLensDistortionState State);
 
+	/** Reset the distortion state to defaults to represent "no distortion" */
+	UFUNCTION(BlueprintCallable, Category = "Lens Component")
+	void ClearDistortionState();
+
 	/** Get the original (not adjusted for overscan) focal length of the camera */
 	UFUNCTION(BlueprintPure, Category = "Lens Component")
 	float GetOriginalFocalLength() const;
diff --git a/Engine/Plugins/VirtualProduction/LiveLinkLens/Source/LiveLinkLens/Private/LiveLinkLensController.cpp b/Engine/Plugins/VirtualProduction/LiveLinkLens/Source/LiveLinkLens/Private/LiveLinkLensController.cpp
--- a/Engine/Plugins/VirtualProduction/LiveLinkLens/Source/LiveLinkLens/Private/LiveLinkLensController.cpp
+++ b/Engine/Plugins/VirtualProduction/LiveLinkLens/Source/LiveLinkLens/Private/LiveLinkLensController.cpp
@@ -54,6 +54,18 @@ void ULiveLinkLensController::Tick(float DeltaTime, const FLiveLinkSubjectFrameD
 	}
 }
 
+void ULiveLinkLensController::Cleanup()
+{
+	if (ULensComponent* const LensComponent = Cast<ULensComponent>(GetAttachedComponent()))
+	{
+		// If the lens component is currently receiving its distortion state from this controller, reset that state to zero
+		if (LensComponent->GetDistortionSource() == EDistortionSource::LiveLinkLensSubject)
+		{
+			LensComponent->ClearDistortionState();
+		}
+	}
+}
+
 bool ULiveLinkLensController::IsRoleSupported(const TSubclassOf<ULiveLinkRole>& RoleToSupport)
 {
 	return RoleToSupport == ULiveLinkLensRole::StaticClass();
diff --git a/Engine/Plugins/VirtualProduction/LiveLinkLens/Source/LiveLinkLens/Public/LiveLinkLensController.h b/Engine/Plugins/VirtualProduction/LiveLinkLens/Source/LiveLinkLens/Public/LiveLinkLensController.h
--- a/Engine/Plugins/VirtualProduction/LiveLinkLens/Source/LiveLinkLens/Public/LiveLinkLensController.h
+++ b/Engine/Plugins/VirtualProduction/LiveLinkLens/Source/LiveLinkLens/Public/LiveLinkLensController.h
@@ -23,6 +23,7 @@ public:
 	virtual bool IsRoleSupported(const TSubclassOf<ULiveLinkRole>& RoleToSupport) override;
 	virtual TSubclassOf<UActorComponent> GetDesiredComponentClass() const override;
 	virtual void SetAttachedComponent(UActorComponent* ActorComponent) override;
+	virtual void Cleanup() override;
 	//~ End ULiveLinkControllerBase interface
 
 	//~ Begin UObject interface
