Fix registration of camera calibration checkerboard placeable item

diff --git a/Engine/Plugins/VirtualProduction/CameraCalibrationCore/Source/CameraCalibrationCoreEditor/Private/CameraCalibrationCoreEditorModule.cpp b/Engine/Plugins/VirtualProduction/CameraCalibrationCore/Source/CameraCalibrationCoreEditor/Private/CameraCalibrationCoreEditorModule.cpp
--- a/Engine/Plugins/VirtualProduction/CameraCalibrationCore/Source/CameraCalibrationCoreEditor/Private/CameraCalibrationCoreEditorModule.cpp
+++ b/Engine/Plugins/VirtualProduction/CameraCalibrationCore/Source/CameraCalibrationCoreEditor/Private/CameraCalibrationCoreEditorModule.cpp
@@ -8,6 +8,7 @@
 #include "ActorFactories/ActorFactory.h"
 #include "CalibrationPointComponent.h"
 #include "CalibrationPointComponentDetails.h"
+#include "CameraCalibrationCheckerboard.h"
 #include "CameraCalibrationCoreEditorStyle.h"
 #include "Editor.h"
 #include "IPlacementModeModule.h"
@@ -113,24 +114,15 @@ void FCameraCalibrationCoreEditorModule::RegisterPlacementModeItems()
 		}
 
 		// Register the Checkerboard
-		{
-			FAssetData CheckerboardAssetData(
-				TEXT("/Script/CameraCalibrationCore.CameraCalibrationCheckerboard"),
-				TEXT("/CameraCalibrationCore"),
-				TEXT("CameraCalibrationCheckerboard"),
-				FTopLevelAssetPath(TEXT("/Script/Engine"), TEXT("Actor"))
-			);
-		
-			PlaceActors.Add(IPlacementModeModule::Get().RegisterPlaceableItem(Info->UniqueHandle, MakeShared<FPlaceableItem>(
-				*UActorFactory::StaticClass(),
-				CheckerboardAssetData,
-				NAME_None,
-				NAME_None,
-				TOptional<FLinearColor>(),
-				TOptional<int32>(),
-				NSLOCTEXT("PlacementMode", "Checkerboard", "Checkerboard")
-			)));
-		}
+		PlaceActors.Add(IPlacementModeModule::Get().RegisterPlaceableItem(Info->UniqueHandle, MakeShared<FPlaceableItem>(
+			*ACameraCalibrationCheckerboard::StaticClass(),
+			FAssetData(ACameraCalibrationCheckerboard::StaticClass()),
+			NAME_None,
+			NAME_None,
+			TOptional<FLinearColor>(),
+			TOptional<int32>(),
+			NSLOCTEXT("PlacementMode", "Checkerboard", "Checkerboard")
+		)));
 	};
 
 	if (FApp::CanEverRender())
