Fix editor crash when using ApplyNodalOffset BP node

diff --git a/Engine/Plugins/VirtualProduction/CameraCalibrationCore/Source/CameraCalibrationCore/Private/LensComponent.cpp b/Engine/Plugins/VirtualProduction/CameraCalibrationCore/Source/CameraCalibrationCore/Private/LensComponent.cpp
--- a/Engine/Plugins/VirtualProduction/CameraCalibrationCore/Source/CameraCalibrationCore/Private/LensComponent.cpp
+++ b/Engine/Plugins/VirtualProduction/CameraCalibrationCore/Source/CameraCalibrationCore/Private/LensComponent.cpp
@@ -416,7 +416,7 @@ void ULensComponent::ApplyNodalOffset(USceneComponent* ComponentToOffset, bool b
 	}
 
 	// Cache the original transform before applying the offset, so that nodal offset can potentially be re-evaluated in the future
-	OriginalTrackedComponentTransform = TrackedComponent.Get()->GetRelativeTransform();
+	OriginalTrackedComponentTransform = ComponentToOffset->GetRelativeTransform();
 
 	ComponentToOffset->AddLocalOffset(Offset.LocationOffset);
 	ComponentToOffset->AddLocalRotation(Offset.RotationOffset);
