Fix CIS non-unity compile error in CameraCalibrationSubsystem

diff --git a/Engine/Plugins/VirtualProduction/CameraCalibrationCore/Source/CameraCalibrationCore/Private/CameraCalibrationSubsystem.cpp b/Engine/Plugins/VirtualProduction/CameraCalibrationCore/Source/CameraCalibrationCore/Private/CameraCalibrationSubsystem.cpp
--- a/Engine/Plugins/VirtualProduction/CameraCalibrationCore/Source/CameraCalibrationCore/Private/CameraCalibrationSubsystem.cpp
+++ b/Engine/Plugins/VirtualProduction/CameraCalibrationCore/Source/CameraCalibrationCore/Private/CameraCalibrationSubsystem.cpp
@@ -6,6 +6,7 @@
 #include "CameraCalibrationSettings.h"
 #include "CineCameraComponent.h"
 #include "Engine/TimecodeProvider.h"
+#include "GameFramework/Actor.h"
 #include "LensComponent.h"
 #include "Misc/CoreDelegates.h"
 #include "UObject/UObjectIterator.h"
