Fix LensComponent MovieScene sections not displaying channels in Sequencer UI because no channel proxy was created on Load.

diff --git a/Engine/Plugins/VirtualProduction/CameraCalibrationCore/Source/CameraCalibrationCoreMovieScene/Private/MovieSceneLensComponentSection.cpp b/Engine/Plugins/VirtualProduction/CameraCalibrationCore/Source/CameraCalibrationCoreMovieScene/Private/MovieSceneLensComponentSection.cpp
--- a/Engine/Plugins/VirtualProduction/CameraCalibrationCore/Source/CameraCalibrationCoreMovieScene/Private/MovieSceneLensComponentSection.cpp
+++ b/Engine/Plugins/VirtualProduction/CameraCalibrationCore/Source/CameraCalibrationCoreMovieScene/Private/MovieSceneLensComponentSection.cpp
@@ -54,6 +54,9 @@ void UMovieSceneLensComponentSection::Initialize(ULensComponent* Component)
 	// Cache a pointer to the recorded component to retrieve its data later during RecordFrame()
 	RecordedComponent = Component;
 
+	// Cache the Lens Model used by the recorded lens component
+	LensModelClass = RecordedComponent->GetLensModel();
+
 	// Duplicate the recorded LensFile asset and save the duplicate in this section
  	if (ULensFile* LensFile = Component->GetLensFile())
  	{
@@ -65,6 +68,20 @@ void UMovieSceneLensComponentSection::Initialize(ULensComponent* Component)
 	CreateChannelProxy();
 }
 
+void UMovieSceneLensComponentSection::PostEditImport()
+{
+	Super::PostEditImport();
+
+	UpdateChannelProxy();
+}
+
+void UMovieSceneLensComponentSection::PostLoad()
+{
+	Super::PostLoad();
+
+	UpdateChannelProxy();
+}
+
 #if WITH_EDITOR
 void UMovieSceneLensComponentSection::AddChannelWithEditor(FMovieSceneChannelProxyData& ChannelProxyData, FMovieSceneFloatChannel& Channel, FText GroupName, FText ChannelName, int32 SortOrder)
 {
@@ -78,18 +95,37 @@ void UMovieSceneLensComponentSection::AddChannelWithEditor(FMovieSceneChannelPro
 
 void UMovieSceneLensComponentSection::CreateChannelProxy()
 {
-	if (const ULensComponent* const LensComponent = RecordedComponent.Get())
+	// Add channels for distortion parameters
+	if (LensModelClass)
 	{
-		FMovieSceneChannelProxyData Channels;
+		if (const ULensModel* const LensModel = LensModelClass->GetDefaultObject<ULensModel>())
+		{
+			const int32 ParameterCount = LensModel->GetNumParameters();
+			DistortionParameterChannels.SetNum(ParameterCount);
+		}
+	}
+
+	// Add channels for FxFy
+	FxFyChannels.SetNum(2);
+
+	// Add channels for Image Center
+	ImageCenterChannels.SetNum(2);
 
-		// Add channels for distortion parameters
-		if (TSubclassOf<ULensModel> LensModelClass = LensComponent->GetLensModel())
+	UpdateChannelProxy();
+}
+
+void UMovieSceneLensComponentSection::UpdateChannelProxy()
+{
+	FMovieSceneChannelProxyData Channels;
+
+	// Add channels for distortion parameters
+	if (LensModelClass)
+	{
+		if (const ULensModel* const LensModel = LensModelClass->GetDefaultObject<ULensModel>())
 		{
-			if (const ULensModel* const LensModel = LensModelClass->GetDefaultObject<ULensModel>())
+			const int32 ParameterCount = LensModel->GetNumParameters();
+			if (ParameterCount == DistortionParameterChannels.Num())
 			{
-				const int32 ParameterCount = LensModel->GetNumParameters();
-				DistortionParameterChannels.SetNum(ParameterCount);
-
 #if WITH_EDITOR
 				TArray<FText> ParameterDisplayNames = LensModel->GetParameterDisplayNames();
 				const FText ParameterGroupName = LOCTEXT("ParameterGroupName", "Distortion Parameters");
@@ -105,10 +141,11 @@ void UMovieSceneLensComponentSection::CreateChannelProxy()
 				}
 			}
 		}
+	}
 
-		// Add channels for FxFy
-		FxFyChannels.SetNum(2);
-
+	// Add channels for FxFy
+	if (FxFyChannels.Num() == 2)
+	{
 #if WITH_EDITOR
 		const FText FxFyGroupName = LOCTEXT("FxFyGroupName", "FxFy");
 		constexpr int32 FxFySortOrder = 1000;
@@ -118,10 +155,11 @@ void UMovieSceneLensComponentSection::CreateChannelProxy()
 		Channels.Add(FxFyChannels[0]);
 		Channels.Add(FxFyChannels[1]);
 #endif // WITH_EDITOR
+	}
 
-		// Add channels for Image Center
-		ImageCenterChannels.SetNum(2);
-
+	// Add channels for Image Center
+	if (ImageCenterChannels.Num() == 2)
+	{
 #if WITH_EDITOR
 		const FText ImageCenterGroupName = LOCTEXT("CxCyGroupName", "Image Center");
 		constexpr int32 ImageCenterSortOrder = 2000;
@@ -131,9 +169,9 @@ void UMovieSceneLensComponentSection::CreateChannelProxy()
 		Channels.Add(ImageCenterChannels[0]);
 		Channels.Add(ImageCenterChannels[1]);
 #endif // WITH_EDITOR
-
-		ChannelProxy = MakeShared<FMovieSceneChannelProxy>(MoveTemp(Channels));
 	}
+
+	ChannelProxy = MakeShared<FMovieSceneChannelProxy>(MoveTemp(Channels));
 }
 
 void UMovieSceneLensComponentSection::RecordFrame(FFrameNumber FrameNumber)
diff --git a/Engine/Plugins/VirtualProduction/CameraCalibrationCore/Source/CameraCalibrationCoreMovieScene/Public/MovieSceneLensComponentSection.h b/Engine/Plugins/VirtualProduction/CameraCalibrationCore/Source/CameraCalibrationCoreMovieScene/Public/MovieSceneLensComponentSection.h
--- a/Engine/Plugins/VirtualProduction/CameraCalibrationCore/Source/CameraCalibrationCoreMovieScene/Public/MovieSceneLensComponentSection.h
+++ b/Engine/Plugins/VirtualProduction/CameraCalibrationCore/Source/CameraCalibrationCoreMovieScene/Public/MovieSceneLensComponentSection.h
@@ -41,10 +41,18 @@ public:
 	/** Record the current value of each distortion state property to its appropriate track */
 	void RecordFrame(FFrameNumber FrameNumber);
 
+	//~ Begin UObject interface
+	virtual void PostEditImport() override;
+	virtual void PostLoad() override;
+	//~ End UObject interface
+
 private:
 	/** Create the channels in which distortion state values will be recorded */
 	void CreateChannelProxy();
 
+	/** Add the distortion state channels to the channel proxy */
+	void UpdateChannelProxy();
+
 #if WITH_EDITOR
 	void AddChannelWithEditor(FMovieSceneChannelProxyData& ChannelProxyData, FMovieSceneFloatChannel& Channel, FText GroupName, FText ChannelName, int32 SortOrder);
 #endif // WITH_EDITOR
@@ -75,6 +83,10 @@ private:
 	UPROPERTY()
 	TArray<FMovieSceneFloatChannel> ImageCenterChannels;
 
+	/** The Lens Model used by the recorded LensComponent */
+	UPROPERTY()
+	TSubclassOf<ULensModel> LensModelClass;
+
 private:
 	TWeakObjectPtr<ULensComponent> RecordedComponent;
 };
