CameraCalibration: Fix STMap vertical flip calculation

diff --git a/Engine/Plugins/VirtualProduction/CameraCalibrationCore/Shaders/Private/DistortionSTMapProcessor.usf b/Engine/Plugins/VirtualProduction/CameraCalibrationCore/Shaders/Private/DistortionSTMapProcessor.usf
--- a/Engine/Plugins/VirtualProduction/CameraCalibrationCore/Shaders/Private/DistortionSTMapProcessor.usf
+++ b/Engine/Plugins/VirtualProduction/CameraCalibrationCore/Shaders/Private/DistortionSTMapProcessor.usf
@@ -33,12 +33,12 @@ RWTexture2D<float2> OutDistortionDisplacementMap;
 void DistortionSTMapProcessorMainCS(uint3 ThreadId : SV_DispatchThreadID)
 {
 	float2 UV = TexelSize * (ThreadId.xy + 0.5f);
-
-#if PIXEL_ORIGIN == PIXEL_ORIGIN_TOPLEFT
 	float4 UVMapping = DistortionSTMap.SampleLevel(DistortionSTMapSampler, UV, 0);
-#elif PIXEL_ORIGIN == PIXEL_ORIGIN_BOTTOMLEFT
-	float2 MirroredUV = float2(UV.x, 1.0f - UV.y);
-	float4 UVMapping = DistortionSTMap.SampleLevel(DistortionSTMapSampler, MirroredUV, 0);
+
+	// If the ST Map is flipped vertically, invert the green and alpha channels
+#if PIXEL_ORIGIN == PIXEL_ORIGIN_BOTTOMLEFT
+	UVMapping.g = 1 - UVMapping.g;
+	UVMapping.a = 1 - UVMapping.a;
 #endif
 
 #if UNDISTORTION_CHANNELS == UNDISTORTION_CHANNELS_RG
diff --git a/Engine/Plugins/VirtualProduction/CameraCalibrationCore/Source/CameraCalibrationCore/Private/LensDistortionModelHandlerBase.cpp b/Engine/Plugins/VirtualProduction/CameraCalibrationCore/Source/CameraCalibrationCore/Private/LensDistortionModelHandlerBase.cpp
--- a/Engine/Plugins/VirtualProduction/CameraCalibrationCore/Source/CameraCalibrationCore/Private/LensDistortionModelHandlerBase.cpp
+++ b/Engine/Plugins/VirtualProduction/CameraCalibrationCore/Source/CameraCalibrationCore/Private/LensDistortionModelHandlerBase.cpp
@@ -72,6 +72,7 @@ void ULensDistortionModelHandlerBase::CreateDisplacementMaps(const FIntPoint Dis
 	{
 		DisplacementMap->RenderTargetFormat = ETextureRenderTargetFormat::RTF_RG16f;
 		DisplacementMap->ClearColor = FLinearColor::Black;
+		DisplacementMap->Filter = TF_Bilinear;
 		DisplacementMap->AddressX = TA_Clamp;
 		DisplacementMap->AddressY = TA_Clamp;
 		DisplacementMap->bAutoGenerateMips = false;
diff --git a/Engine/Plugins/VirtualProduction/CameraCalibrationCore/Source/CameraCalibrationCore/Private/LensFile.cpp b/Engine/Plugins/VirtualProduction/CameraCalibrationCore/Source/CameraCalibrationCore/Private/LensFile.cpp
--- a/Engine/Plugins/VirtualProduction/CameraCalibrationCore/Source/CameraCalibrationCore/Private/LensFile.cpp
+++ b/Engine/Plugins/VirtualProduction/CameraCalibrationCore/Source/CameraCalibrationCore/Private/LensFile.cpp
@@ -30,6 +30,7 @@ namespace LensFileUtils
 		UTextureRenderTarget2D* NewRenderTarget2D = NewObject<UTextureRenderTarget2D>(Outer, MakeUniqueObjectName(Outer, UTextureRenderTarget2D::StaticClass(), TEXT("LensDisplacementMap")), RF_Public);
 		NewRenderTarget2D->RenderTargetFormat = ETextureRenderTargetFormat::RTF_RG16f;
 		NewRenderTarget2D->ClearColor = FLinearColor(0.5f, 0.5f, 0.5f, 0.5f);
+		NewRenderTarget2D->Filter = TF_Bilinear;
 		NewRenderTarget2D->AddressX = TA_Clamp;
 		NewRenderTarget2D->AddressY = TA_Clamp;
 		NewRenderTarget2D->bAutoGenerateMips = false;
