Expose distortion state properties to cinematics

diff --git a/Engine/Plugins/VirtualProduction/CameraCalibrationCore/Source/CameraCalibrationCore/Public/LensData.h b/Engine/Plugins/VirtualProduction/CameraCalibrationCore/Source/CameraCalibrationCore/Public/LensData.h
--- a/Engine/Plugins/VirtualProduction/CameraCalibrationCore/Source/CameraCalibrationCore/Public/LensData.h
+++ b/Engine/Plugins/VirtualProduction/CameraCalibrationCore/Source/CameraCalibrationCore/Public/LensData.h
@@ -66,7 +66,7 @@ struct CAMERACALIBRATIONCORE_API FDistortionInfo
 
 public:
 	/** Generic array of floating-point lens distortion parameters */
-	UPROPERTY(EditAnywhere, BlueprintReadWrite, Category = "Distortion")
+	UPROPERTY(Interp, EditAnywhere, BlueprintReadWrite, Category = "Distortion")
 	TArray<float> Parameters;
 };
 
@@ -83,7 +83,7 @@ struct CAMERACALIBRATIONCORE_API FFocalLengthInfo
 public:
 
 	/** Value expected to be normalized (unitless) */
-	UPROPERTY(EditAnywhere, BlueprintReadWrite, Category = "Camera")
+	UPROPERTY(Interp, EditAnywhere, BlueprintReadWrite, Category = "Camera")
 	FVector2D FxFy = FVector2D(1.0f, (16.0f / 9.0f));
 };
 
@@ -115,7 +115,7 @@ struct CAMERACALIBRATIONCORE_API FImageCenterInfo
 
 public:
 	/** Value expected to be normalized [0,1] */
-	UPROPERTY(EditAnywhere, BlueprintReadWrite, Category = "Camera", meta = (DisplayName = "Image Center"))
+	UPROPERTY(Interp, EditAnywhere, BlueprintReadWrite, Category = "Camera", meta = (DisplayName = "Image Center"))
 	FVector2D PrincipalPoint = FVector2D(0.5f, 0.5f);
 };
 
