CameraCalibration: Add anamorphic lens model class and handler

diff --git a/Engine/Plugins/VirtualProduction/CameraCalibrationCore/Source/CameraCalibrationCore/Private/AnamorphicLensDistortionModelHandler.cpp b/Engine/Plugins/VirtualProduction/CameraCalibrationCore/Source/CameraCalibrationCore/Private/AnamorphicLensDistortionModelHandler.cpp
new file mode 100644
--- /dev/null
+++ b/Engine/Plugins/VirtualProduction/CameraCalibrationCore/Source/CameraCalibrationCore/Private/AnamorphicLensDistortionModelHandler.cpp
@@ -0,0 +1,122 @@
+// Copyright Epic Games, Inc. All Rights Reserved.
+
+#include "AnamorphicLensDistortionModelHandler.h"
+
+#include "CameraCalibrationSettings.h"
+#include "Engine/TextureRenderTarget2D.h"
+
+void UAnamorphicLensDistortionModelHandler::InitializeHandler()
+{
+	LensModelClass = UAnamorphicLensModel::StaticClass();
+}
+
+FVector2D UAnamorphicLensDistortionModelHandler::ComputeDistortedUV(const FVector2D& InUndistortedUV) const
+{
+	// The following implementation is based on the 3DE4 Anamorphic - Standard, Degree 4 model for distortion
+
+	// Pre-compute some values that will be useful later
+	const FVector2D FilmbackSize = FVector2D(CameraFilmback.SensorWidth, CameraFilmback.SensorHeight);
+
+	// Guard against divide-by-zero errors
+	if (FMath::IsNearlyZero(FilmbackSize.X) || FMath::IsNearlyZero(FilmbackSize.Y) || FMath::IsNearlyZero(AnamorphicParameters.PixelAspect))
+	{
+		return InUndistortedUV;
+	}
+
+	const double FilmbackRadius = 0.5f * FMath::Sqrt((FilmbackSize.X * FilmbackSize.X) + (FilmbackSize.Y * FilmbackSize.Y));
+
+	const double MountAngleRadians = FMath::DegreesToRadians(AnamorphicParameters.LensRotation);
+	const double CosMountAngle = FMath::Cos(MountAngleRadians);
+	const double SinMountAngle = FMath::Sin(MountAngleRadians);
+
+	// Step 1: Transform input UVs into diagonally-normalized coordinates, based on the width and height of the filmback/image
+	FVector2D DNCoordinates = (InUndistortedUV - 0.5f) * (FilmbackSize / FilmbackRadius);
+
+	// Step 2: Scale the x-coordinate by the inverse of the pixel aspect ratio
+	DNCoordinates.X = DNCoordinates.X / AnamorphicParameters.PixelAspect;
+
+	// Step 3: Rotate x and y by the inverse of the mounting angle
+	DNCoordinates.X = (DNCoordinates.X * CosMountAngle) + (DNCoordinates.Y * SinMountAngle);
+	DNCoordinates.Y = (DNCoordinates.Y * CosMountAngle) - (DNCoordinates.X * SinMountAngle);
+
+	// Step 4: Compute the anamorphic distortion for x and y, to the 4th degree
+	const double DNRadiusSq = (DNCoordinates.X * DNCoordinates.X) + (DNCoordinates.Y * DNCoordinates.Y);
+	const double DNRadius4th = DNRadiusSq * DNRadiusSq;
+	const double Phi = FMath::Atan2(DNCoordinates.Y, DNCoordinates.X);
+
+	FVector2D DNCoordinatesDistorted;
+	DNCoordinatesDistorted.X = DNCoordinates.X * (1 + (AnamorphicParameters.CX02 * DNRadiusSq) + (AnamorphicParameters.CX04 * DNRadius4th) + (AnamorphicParameters.CX22 * DNRadiusSq * FMath::Cos(2 * Phi)) + (AnamorphicParameters.CX24 * DNRadius4th * FMath::Cos(2 * Phi)) + (AnamorphicParameters.CX44 * DNRadius4th * FMath::Cos(4 * Phi)));
+	DNCoordinatesDistorted.Y = DNCoordinates.Y * (1 + (AnamorphicParameters.CY02 * DNRadiusSq) + (AnamorphicParameters.CY04 * DNRadius4th) + (AnamorphicParameters.CY22 * DNRadiusSq * FMath::Cos(2 * Phi)) + (AnamorphicParameters.CY24 * DNRadius4th * FMath::Cos(2 * Phi)) + (AnamorphicParameters.CY44 * DNRadius4th * FMath::Cos(4 * Phi)));
+
+	// Step 5: Scale the x-coordinate by the pixel aspect ratio, and the x- and y-coordinates by the lens-breathing squeeze factors
+	DNCoordinatesDistorted.X = DNCoordinatesDistorted.X * AnamorphicParameters.PixelAspect * AnamorphicParameters.SqueezeX;
+	DNCoordinatesDistorted.Y = DNCoordinatesDistorted.Y * AnamorphicParameters.SqueezeY;
+
+	// Step 6: Rotate x and y by the mounting angle
+	DNCoordinates.X = (DNCoordinates.X * CosMountAngle) - (DNCoordinates.Y * SinMountAngle);
+	DNCoordinates.Y = (DNCoordinates.Y * CosMountAngle) + (DNCoordinates.X * SinMountAngle);
+
+	// Step 7: Transform diagonally normalized coordinates back into unit coordinates
+	const FVector2D DistortedUV = ((DNCoordinatesDistorted * FilmbackRadius) / FilmbackSize) + 0.5f;
+
+	return DistortedUV;
+}
+
+void UAnamorphicLensDistortionModelHandler::InitDistortionMaterials()
+{
+	if (DistortionPostProcessMID == nullptr)
+	{
+		UMaterialInterface* DistortionMaterialParent = GetDefault<UCameraCalibrationSettings>()->GetDefaultDistortionMaterial(this->StaticClass());
+		DistortionPostProcessMID = UMaterialInstanceDynamic::Create(DistortionMaterialParent, this);
+	}
+
+	// TODO: Add an undistortion displacement map MID once the material is authored
+
+	if (DistortionDisplacementMapMID == nullptr)
+	{
+		UMaterialInterface* MaterialParent = GetDefault<UCameraCalibrationSettings>()->GetDefaultDistortionDisplacementMaterial(this->StaticClass());
+		DistortionDisplacementMapMID = UMaterialInstanceDynamic::Create(MaterialParent, this);
+	}
+
+	// TODO: Add the UndistortionDisplacementMap parameter
+	DistortionPostProcessMID->SetTextureParameterValue("DistortionDisplacementMap", DistortionDisplacementMapRT);
+
+	SetDistortionState(CurrentState);
+}
+
+void UAnamorphicLensDistortionModelHandler::UpdateMaterialParameters()
+{
+	//Helper function to set material parameters of an MID
+	const auto SetDistortionMaterialParameters = [this](UMaterialInstanceDynamic* const MID)
+	{
+		MID->SetScalarParameterValue("pixel_aspect", AnamorphicParameters.PixelAspect);
+
+ 		MID->SetScalarParameterValue("w_fb", CameraFilmback.SensorWidth);
+ 		MID->SetScalarParameterValue("h_fb", CameraFilmback.SensorHeight);
+
+		MID->SetScalarParameterValue("cx02", AnamorphicParameters.CX02);
+		MID->SetScalarParameterValue("cx04", AnamorphicParameters.CX04);
+		MID->SetScalarParameterValue("cx22", AnamorphicParameters.CX22);
+		MID->SetScalarParameterValue("cx24", AnamorphicParameters.CX24);
+		MID->SetScalarParameterValue("cx44", AnamorphicParameters.CX44);
+
+		MID->SetScalarParameterValue("cy02", AnamorphicParameters.CY02);
+		MID->SetScalarParameterValue("cy04", AnamorphicParameters.CY04);
+		MID->SetScalarParameterValue("cy22", AnamorphicParameters.CY22);
+		MID->SetScalarParameterValue("cy24", AnamorphicParameters.CY24);
+		MID->SetScalarParameterValue("cy44", AnamorphicParameters.CY44);
+
+		MID->SetScalarParameterValue("sx", AnamorphicParameters.SqueezeX);
+		MID->SetScalarParameterValue("sy", AnamorphicParameters.SqueezeY);
+
+		MID->SetScalarParameterValue("phi_mnt", AnamorphicParameters.LensRotation);
+	};
+
+	// TODO: Update the undistortion material parameters
+	SetDistortionMaterialParameters(DistortionDisplacementMapMID);
+}
+
+void UAnamorphicLensDistortionModelHandler::InterpretDistortionParameters()
+{
+	LensModelClass->GetDefaultObject<ULensModel>()->FromArray<FAnamorphicDistortionParameters>(CurrentState.DistortionInfo.Parameters, AnamorphicParameters);
+}
diff --git a/Engine/Plugins/VirtualProduction/CameraCalibrationCore/Source/CameraCalibrationCore/Private/CameraCalibrationCoreModule.cpp b/Engine/Plugins/VirtualProduction/CameraCalibrationCore/Source/CameraCalibrationCore/Private/CameraCalibrationCoreModule.cpp
--- a/Engine/Plugins/VirtualProduction/CameraCalibrationCore/Source/CameraCalibrationCore/Private/CameraCalibrationCoreModule.cpp
+++ b/Engine/Plugins/VirtualProduction/CameraCalibrationCore/Source/CameraCalibrationCore/Private/CameraCalibrationCoreModule.cpp
@@ -12,6 +12,7 @@
 #include "Misc/App.h"
 #include "Misc/CoreDelegates.h"
 #include "Misc/Paths.h"
+#include "Models/AnamorphicLensModel.h"
 #include "Models/SphericalLensModel.h"
 #include "Modules/ModuleManager.h"
 #include "ShaderCore.h"
@@ -119,6 +120,7 @@ void FCameraCalibrationCoreModule::RegisterDistortionModels()
 		// Register all lens models defined in this module
 		UCameraCalibrationSubsystem* SubSystem = GEngine->GetEngineSubsystem<UCameraCalibrationSubsystem>();
 		SubSystem->RegisterDistortionModel(USphericalLensModel::StaticClass());
+		SubSystem->RegisterDistortionModel(UAnamorphicLensModel::StaticClass());
 	};
 
 	if (FApp::CanEverRender())
@@ -142,6 +144,7 @@ void FCameraCalibrationCoreModule::UnregisterDistortionModels()
 		if (UCameraCalibrationSubsystem* SubSystem = GEngine->GetEngineSubsystem<UCameraCalibrationSubsystem>())
 		{
 			SubSystem->UnregisterDistortionModel(USphericalLensModel::StaticClass());
+			SubSystem->UnregisterDistortionModel(UAnamorphicLensModel::StaticClass());
 		}
 	}
 }
diff --git a/Engine/Plugins/VirtualProduction/CameraCalibrationCore/Source/CameraCalibrationCore/Private/CameraCalibrationSettings.cpp b/Engine/Plugins/VirtualProduction/CameraCalibrationCore/Source/CameraCalibrationCore/Private/CameraCalibrationSettings.cpp
--- a/Engine/Plugins/VirtualProduction/CameraCalibrationCore/Source/CameraCalibrationCore/Private/CameraCalibrationSettings.cpp
+++ b/Engine/Plugins/VirtualProduction/CameraCalibrationCore/Source/CameraCalibrationCore/Private/CameraCalibrationSettings.cpp
@@ -1,18 +1,25 @@
 // Copyright Epic Games, Inc. All Rights Reserved.
 
 #include "CameraCalibrationSettings.h"
-#include "CameraCalibrationSubsystem.h"
 
+#include "AnamorphicLensDistortionModelHandler.h"
+#include "CameraCalibrationSubsystem.h"
 #include "Engine/Engine.h"
 #include "Misc/CoreDelegates.h"
 #include "SphericalLensDistortionModelHandler.h"
 
 UCameraCalibrationSettings::UCameraCalibrationSettings()
 {
+	// Default materials to use for the spherical lens model
 	DefaultUndistortionDisplacementMaterials.Add(USphericalLensDistortionModelHandler::StaticClass(), TSoftObjectPtr<UMaterialInterface>(FSoftObjectPath(TEXT("/CameraCalibrationCore/Materials/M_SphericalUndistortionDisplacementMap.M_SphericalUndistortionDisplacementMap"))));
 	DefaultDistortionDisplacementMaterials.Add(USphericalLensDistortionModelHandler::StaticClass(), TSoftObjectPtr<UMaterialInterface>(FSoftObjectPath(TEXT("/CameraCalibrationCore/Materials/M_SphericalDistortionDisplacementMap.M_SphericalDistortionDisplacementMap"))));
 	DefaultDistortionMaterials.Add(USphericalLensDistortionModelHandler::StaticClass(), TSoftObjectPtr<UMaterialInterface>(FSoftObjectPath(TEXT("/CameraCalibrationCore/Materials/M_DistortionPostProcess.M_DistortionPostProcess"))));
 
+	// Default materials to use for the anamorphic lens model
+	// TODO: Add default undistortion displacement map material once it is authored
+	DefaultDistortionDisplacementMaterials.Add(UAnamorphicLensDistortionModelHandler::StaticClass(), TSoftObjectPtr<UMaterialInterface>(FSoftObjectPath(TEXT("/CameraCalibrationCore/Materials/M_AnamorphicDistortionDisplacementMap.M_AnamorphicDistortionDisplacementMap"))));
+	DefaultDistortionMaterials.Add(UAnamorphicLensDistortionModelHandler::StaticClass(), TSoftObjectPtr<UMaterialInterface>(FSoftObjectPath(TEXT("/CameraCalibrationCore/Materials/M_DistortionPostProcess.M_DistortionPostProcess"))));
+
 #if WITH_EDITORONLY_DATA
 	auto InitOverlayOverrides = [this]()
 	{
diff --git a/Engine/Plugins/VirtualProduction/CameraCalibrationCore/Source/CameraCalibrationCore/Private/LensComponent.cpp b/Engine/Plugins/VirtualProduction/CameraCalibrationCore/Source/CameraCalibrationCore/Private/LensComponent.cpp
--- a/Engine/Plugins/VirtualProduction/CameraCalibrationCore/Source/CameraCalibrationCore/Private/LensComponent.cpp
+++ b/Engine/Plugins/VirtualProduction/CameraCalibrationCore/Source/CameraCalibrationCore/Private/LensComponent.cpp
@@ -142,6 +142,7 @@ void ULensComponent::TickComponent(float DeltaTime, ELevelTick TickType, FActorC
 			case EDistortionSource::Manual:
 			{
 				LensDistortionHandler->SetDistortionState(DistortionState);
+				LensDistortionHandler->SetCameraFilmback(CineCameraComponent->Filmback);
 
 				//Recompute overscan factor for the distortion state
 				float OverscanFactor = LensDistortionHandler->ComputeOverscanFactor();
@@ -591,11 +592,14 @@ void ULensComponent::SetLensModel(TSubclassOf<ULensModel> Model)
 
 	if (LensModel)
 	{
-		const uint32 NumDistortionParameters = LensModel->GetDefaultObject<ULensModel>()->GetNumParameters();
-		DistortionState.DistortionInfo.Parameters.Init(0.0f, NumDistortionParameters);
+		LensModel->GetDefaultObject<ULensModel>()->GetDefaultParameterArray(DistortionState.DistortionInfo.Parameters);
 
 		CreateDistortionHandler();
 	}
+	else
+	{
+		DistortionState.DistortionInfo.Parameters.Empty();
+	}
 }
 
 FLensDistortionState ULensComponent::GetDistortionState() const
@@ -699,8 +703,7 @@ void ULensComponent::ClearDistortionState()
 	FLensDistortionState ZeroDistortionState;
 	if (LensModel)
 	{
-		const uint32 NumDistortionParameters = LensModel->GetDefaultObject<ULensModel>()->GetNumParameters();
-		ZeroDistortionState.DistortionInfo.Parameters.SetNumZeroed(NumDistortionParameters);
+		LensModel->GetDefaultObject<ULensModel>()->GetDefaultParameterArray(ZeroDistortionState.DistortionInfo.Parameters);
 	}
 	DistortionState = ZeroDistortionState;
 }
diff --git a/Engine/Plugins/VirtualProduction/CameraCalibrationCore/Source/CameraCalibrationCore/Private/LensDistortionModelHandlerBase.cpp b/Engine/Plugins/VirtualProduction/CameraCalibrationCore/Source/CameraCalibrationCore/Private/LensDistortionModelHandlerBase.cpp
--- a/Engine/Plugins/VirtualProduction/CameraCalibrationCore/Source/CameraCalibrationCore/Private/LensDistortionModelHandlerBase.cpp
+++ b/Engine/Plugins/VirtualProduction/CameraCalibrationCore/Source/CameraCalibrationCore/Private/LensDistortionModelHandlerBase.cpp
@@ -42,6 +42,17 @@ void ULensDistortionModelHandlerBase::SetDistortionState(const FLensDistortionSt
 	}	
 }
 
+void ULensDistortionModelHandlerBase::SetCameraFilmback(const FCameraFilmbackSettings& InCameraFilmback)
+{
+	if (CameraFilmback != InCameraFilmback)
+	{
+		CameraFilmback = InCameraFilmback;
+		InterpretDistortionParameters();
+
+		bIsDirty = true;
+	}
+}
+
 void ULensDistortionModelHandlerBase::PostInitProperties()
 {
 	Super::PostInitProperties();
@@ -56,8 +67,7 @@ void ULensDistortionModelHandlerBase::PostInitProperties()
 	{
 		if (LensModelClass)
 		{
-			const uint32 NumDistortionParameters = LensModelClass->GetDefaultObject<ULensModel>()->GetNumParameters();
-			CurrentState.DistortionInfo.Parameters.Init(0.0f, NumDistortionParameters);
+			LensModelClass->GetDefaultObject<ULensModel>()->GetDefaultParameterArray(CurrentState.DistortionInfo.Parameters);
 		}
 
 		const FIntPoint DisplacementMapResolution = GetDefault<UCameraCalibrationSettings>()->GetDisplacementMapResolution();
diff --git a/Engine/Plugins/VirtualProduction/CameraCalibrationCore/Source/CameraCalibrationCore/Private/LensFile.cpp b/Engine/Plugins/VirtualProduction/CameraCalibrationCore/Source/CameraCalibrationCore/Private/LensFile.cpp
--- a/Engine/Plugins/VirtualProduction/CameraCalibrationCore/Source/CameraCalibrationCore/Private/LensFile.cpp
+++ b/Engine/Plugins/VirtualProduction/CameraCalibrationCore/Source/CameraCalibrationCore/Private/LensFile.cpp
@@ -474,13 +474,17 @@ bool ULensFile::EvaluateDistortionForParameters(float InFocus, float InZoom, FVe
 	FImageCenterInfo ImageCenter;
 	EvaluateImageCenterParameters(InFocus, InZoom, ImageCenter);
 
+	FCameraFilmbackSettings CameraFilmback;
+	CameraFilmback.SensorWidth = InFilmback.X;
+	CameraFilmback.SensorHeight = InFilmback.Y;
+
 	//No distortion parameters case. Still process to have center shift
 	if (DistortionTable.GetFocusPoints().Num() <= 0)
 	{
-		//Initialize all distortion parameters with 0
+		// Initialize all distortion parameters with their default values
 		FDistortionInfo DistortionPoint;
-		DistortionPoint.Parameters.SetNumZeroed(LensInfo.LensModel.GetDefaultObject()->GetNumParameters());
-		
+		LensInfo.LensModel->GetDefaultObject<ULensModel>()->GetDefaultParameterArray(DistortionPoint.Parameters);
+
 		//Setup handler state based on evaluated parameters. If none were found, no distortion will be returned
 		FLensDistortionState State;
 		State.DistortionInfo.Parameters = MoveTemp(DistortionPoint.Parameters);
@@ -494,6 +498,7 @@ bool ULensFile::EvaluateDistortionForParameters(float InFocus, float InZoom, FVe
 
 		//Updates handler state
 		InLensHandler->SetDistortionState(State);
+		InLensHandler->SetCameraFilmback(CameraFilmback);
 
 		InLensHandler->SetOverscanFactor(1.0f);
 
@@ -509,14 +514,15 @@ bool ULensFile::EvaluateDistortionForParameters(float InFocus, float InZoom, FVe
 		float InterpolatedOverscanFactor = 1.0f;
 
 		//Helper function to compute the current distortion state
-		const auto ProcessDistortionState = [this, &State, InFilmback, InLensHandler](const FDistortionInfo& DistortionInfo, const FFocalLengthInfo& FocalLength, UTextureRenderTarget2D* UndistortionRenderTarget, UTextureRenderTarget2D* DistortionRenderTarget, float& OverscanFactor)
+		const auto ProcessDistortionState = [this, &State, CameraFilmback, InLensHandler](const FDistortionInfo& DistortionInfo, const FFocalLengthInfo& FocalLength, UTextureRenderTarget2D* UndistortionRenderTarget, UTextureRenderTarget2D* DistortionRenderTarget, float& OverscanFactor)
 		{
 			State.DistortionInfo.Parameters = DistortionInfo.Parameters;
 
-			const FVector2D FxFyScale = FVector2D(LensInfo.SensorDimensions.X / InFilmback.X, LensInfo.SensorDimensions.Y / InFilmback.Y);
+			const FVector2D FxFyScale = FVector2D(LensInfo.SensorDimensions.X / CameraFilmback.SensorWidth, LensInfo.SensorDimensions.Y / CameraFilmback.SensorHeight);
 			State.FocalLengthInfo.FxFy = FocalLength.FxFy * FxFyScale;
 
 			InLensHandler->SetDistortionState(State);
+			InLensHandler->SetCameraFilmback(CameraFilmback);
 			InLensHandler->DrawUndistortionDisplacementMap(UndistortionRenderTarget);
 			InLensHandler->DrawDistortionDisplacementMap(DistortionRenderTarget);
 
@@ -674,6 +680,7 @@ bool ULensFile::EvaluateDistortionForParameters(float InFocus, float InZoom, FVe
 		
 		//Sets final blended distortion state
 		InLensHandler->SetDistortionState(State);
+		InLensHandler->SetCameraFilmback(CameraFilmback);
 
 		//Draw resulting undistortion displacement map for evaluation point
 		LensFileRendering::DrawBlendedDisplacementMap(InLensHandler->GetUndistortionDisplacementMap()
@@ -725,6 +732,10 @@ bool ULensFile::EvaluateDistortionForSTMaps(float InFocus, float InZoom, FVector
 		return false;
 	}
 
+	FCameraFilmbackSettings CameraFilmback;
+	CameraFilmback.SensorWidth = InFilmback.X;
+	CameraFilmback.SensorHeight = InFilmback.Y;
+
 	FDisplacementMapBlendingParams Params;
 
 	TArray<UTextureRenderTarget2D*, TInlineAllocator<4>> UndistortionMapSource;
@@ -895,6 +906,8 @@ bool ULensFile::EvaluateDistortionForSTMaps(float InFocus, float InZoom, FVector
 			
 	//Sets final blended distortion state
 	InLensHandler->SetDistortionState(MoveTemp(State));
+	InLensHandler->SetCameraFilmback(CameraFilmback);
+
 	InLensHandler->SetOverscanFactor(InterpolatedOverscanFactor);
 
 	return true;
diff --git a/Engine/Plugins/VirtualProduction/CameraCalibrationCore/Source/CameraCalibrationCore/Private/Models/AnamorphicLensModel.cpp b/Engine/Plugins/VirtualProduction/CameraCalibrationCore/Source/CameraCalibrationCore/Private/Models/AnamorphicLensModel.cpp
new file mode 100644
--- /dev/null
+++ b/Engine/Plugins/VirtualProduction/CameraCalibrationCore/Source/CameraCalibrationCore/Private/Models/AnamorphicLensModel.cpp
@@ -0,0 +1,18 @@
+// Copyright Epic Games, Inc. All Rights Reserved.
+
+#include "Models/AnamorphicLensModel.h"
+
+UScriptStruct* UAnamorphicLensModel::GetParameterStruct() const
+{
+	return FAnamorphicDistortionParameters::StaticStruct();
+}
+
+FName UAnamorphicLensModel::GetModelName() const
+{
+	return FName("3DE4 Anamorphic Standard Degree 4");
+}
+
+FName UAnamorphicLensModel::GetShortModelName() const
+{
+	return TEXT("3DE4 Anamorphic Standard Degree 4");
+}
\ No newline at end of file
diff --git a/Engine/Plugins/VirtualProduction/CameraCalibrationCore/Source/CameraCalibrationCore/Private/Models/LensModel.cpp b/Engine/Plugins/VirtualProduction/CameraCalibrationCore/Source/CameraCalibrationCore/Private/Models/LensModel.cpp
--- a/Engine/Plugins/VirtualProduction/CameraCalibrationCore/Source/CameraCalibrationCore/Private/Models/LensModel.cpp
+++ b/Engine/Plugins/VirtualProduction/CameraCalibrationCore/Source/CameraCalibrationCore/Private/Models/LensModel.cpp
@@ -35,6 +35,12 @@ TArray<FText> ULensModel::GetParameterDisplayNames() const
 
 #endif //WITH_EDITOR
 
+void ULensModel::GetDefaultParameterArray(TArray<float>& OutParameterValues) const
+{
+	FStructOnScope StructOnScope{ GetParameterStruct() };
+	ToArray<FStructOnScope>(StructOnScope, OutParameterValues);
+}
+
 uint32 ULensModel::GetNumParameters() const
 {
 	UScriptStruct* TypeStruct = GetParameterStruct();
diff --git a/Engine/Plugins/VirtualProduction/CameraCalibrationCore/Source/CameraCalibrationCore/Public/AnamorphicLensDistortionModelHandler.h b/Engine/Plugins/VirtualProduction/CameraCalibrationCore/Source/CameraCalibrationCore/Public/AnamorphicLensDistortionModelHandler.h
new file mode 100644
--- /dev/null
+++ b/Engine/Plugins/VirtualProduction/CameraCalibrationCore/Source/CameraCalibrationCore/Public/AnamorphicLensDistortionModelHandler.h
@@ -0,0 +1,29 @@
+// Copyright Epic Games, Inc. All Rights Reserved.
+
+#pragma once
+
+#include "LensDistortionModelHandlerBase.h"
+
+#include "Models/AnamorphicLensModel.h"
+
+#include "AnamorphicLensDistortionModelHandler.generated.h"
+
+/** Lens distortion handler for an Anamorphic lens model that implements the 3DE4 Anamorphic - Standard Degree 4 model */
+UCLASS(BlueprintType)
+class CAMERACALIBRATIONCORE_API UAnamorphicLensDistortionModelHandler : public ULensDistortionModelHandlerBase
+{
+	GENERATED_BODY()
+
+protected:
+	//~ Begin ULensDistortionModelHandlerBase interface
+	virtual void InitializeHandler() override;
+	virtual FVector2D ComputeDistortedUV(const FVector2D& InScreenUV) const override;
+	virtual void InitDistortionMaterials() override;
+	virtual void UpdateMaterialParameters() override;
+	virtual void InterpretDistortionParameters() override;
+	//~ End ULensDistortionModelHandlerBase interface
+
+private:
+	/** Anamorphic lens distortion parameters */
+	FAnamorphicDistortionParameters AnamorphicParameters;
+};
diff --git a/Engine/Plugins/VirtualProduction/CameraCalibrationCore/Source/CameraCalibrationCore/Public/LensDistortionModelHandlerBase.h b/Engine/Plugins/VirtualProduction/CameraCalibrationCore/Source/CameraCalibrationCore/Public/LensDistortionModelHandlerBase.h
--- a/Engine/Plugins/VirtualProduction/CameraCalibrationCore/Source/CameraCalibrationCore/Public/LensDistortionModelHandlerBase.h
+++ b/Engine/Plugins/VirtualProduction/CameraCalibrationCore/Source/CameraCalibrationCore/Public/LensDistortionModelHandlerBase.h
@@ -2,6 +2,7 @@
 
 #pragma once
 
+#include "CineCameraSettings.h"
 #include "CoreMinimal.h"
 #include "LensData.h"
 #include "Materials/MaterialInstanceDynamic.h"
@@ -70,6 +71,9 @@ public:
 	/** Get the current distortion state (the lens model and properties that mathematically represent the distortion characteristics */
 	FLensDistortionState GetCurrentDistortionState() const { return CurrentState; }
 
+	/** Specify the filmback settings of the CineCamera that is being used for distortion */
+	void SetCameraFilmback(const FCameraFilmbackSettings& InCameraFilmback);
+
 	/** Get the post-process MID for the currently specified lens model */
 	UMaterialInstanceDynamic* GetDistortionMID() const { return DistortionPostProcessMID; }
 
@@ -151,6 +155,10 @@ protected:
 	UPROPERTY(EditAnywhere, BlueprintReadOnly, Category = "Distortion", meta = (ShowOnlyInnerProperties))
 	FLensDistortionState CurrentState;
 
+	/** Filmback settings of the camera that is being used for distortion */
+	UPROPERTY(EditAnywhere, BlueprintReadWrite, Category = "Distortion")
+	FCameraFilmbackSettings CameraFilmback;
+
 	/** Display name, used to identify handler in-editor details panels */
 	UPROPERTY(VisibleAnywhere, Category = "Distortion")
 	FString DisplayName;
diff --git a/Engine/Plugins/VirtualProduction/CameraCalibrationCore/Source/CameraCalibrationCore/Public/Models/AnamorphicLensModel.h b/Engine/Plugins/VirtualProduction/CameraCalibrationCore/Source/CameraCalibrationCore/Public/Models/AnamorphicLensModel.h
new file mode 100644
--- /dev/null
+++ b/Engine/Plugins/VirtualProduction/CameraCalibrationCore/Source/CameraCalibrationCore/Public/Models/AnamorphicLensModel.h
@@ -0,0 +1,91 @@
+// Copyright Epic Games, Inc. All Rights Reserved.
+
+#pragma once
+
+#include "LensModel.h"
+
+#include "AnamorphicLensModel.generated.h"
+
+/**
+ * Lens distortion parameters for the 3DE4 Anamorphic - Standard Degree 4 model
+ * All parameters are unitless and represent the coefficients used to undistort a distorted image
+ * For complete model description, see "tde4_ldm_standard.pdf" from https://www.3dequalizer.com/ in the Lens Distortion Plugin Kit v2.8 
+ */
+USTRUCT(BlueprintType)
+struct CAMERACALIBRATIONCORE_API FAnamorphicDistortionParameters
+{
+	GENERATED_BODY()
+
+public:
+	/** Anamorphic Squeeze (the ratio of the filmback size to the size of the rasterized image) */
+	UPROPERTY(EditAnywhere, BlueprintReadWrite, Category = "Distortion")
+	float PixelAspect = 1.0f;
+
+	/** X coefficient of the r^2 term */
+	UPROPERTY(EditAnywhere, BlueprintReadWrite, Category = "Distortion")
+	float CX02 = 0.0f;
+
+	/** X coefficient of the r^4 term */
+	UPROPERTY(EditAnywhere, BlueprintReadWrite, Category = "Distortion")
+	float CX04 = 0.0f;
+
+	/** X coefficient of the r^2*cos(2*phi) term */
+	UPROPERTY(EditAnywhere, BlueprintReadWrite, Category = "Distortion")
+	float CX22 = 0.0f;
+
+	/** X coefficient of the r^4*cos(2*phi) term */
+	UPROPERTY(EditAnywhere, BlueprintReadWrite, Category = "Distortion")
+	float CX24 = 0.0f;
+
+	/** X coefficient of the r^4*cos(4*phi) term */
+	UPROPERTY(EditAnywhere, BlueprintReadWrite, Category = "Distortion")
+	float CX44 = 0.0f;
+
+	/** Y coefficient of the r^2 term */
+	UPROPERTY(EditAnywhere, BlueprintReadWrite, Category = "Distortion")
+	float CY02 = 0.0f;
+
+	/** Y coefficient of the r^4 term */
+	UPROPERTY(EditAnywhere, BlueprintReadWrite, Category = "Distortion")
+	float CY04 = 0.0f;
+
+	/** Y coefficient of the r^2*cos(2*phi) term */
+	UPROPERTY(EditAnywhere, BlueprintReadWrite, Category = "Distortion")
+	float CY22 = 0.0f;
+
+	/** Y coefficient of the r^4*cos(2*phi) term */
+	UPROPERTY(EditAnywhere, BlueprintReadWrite, Category = "Distortion")
+	float CY24 = 0.0f;
+
+	/** Y coefficient of the r^4*cos(4*phi) term */
+	UPROPERTY(EditAnywhere, BlueprintReadWrite, Category = "Distortion")
+	float CY44 = 0.0f;
+
+	/** Squeeze Factor (should be small, relatively close to 1.0) */
+	UPROPERTY(EditAnywhere, BlueprintReadWrite, Category = "Distortion")
+	float SqueezeX = 1.0f;
+
+	/** Squeeze Factor (should be small, relatively close to 1.0) */
+	UPROPERTY(EditAnywhere, BlueprintReadWrite, Category = "Distortion")
+	float SqueezeY = 1.0f;
+
+	/** Lens Rotation in degrees. Represents mounting inaccuracies (should be small, between -2 and +2 degrees) */
+	UPROPERTY(EditAnywhere, BlueprintReadWrite, Category = "Distortion")
+	float LensRotation = 0.0f;
+};
+
+/**
+ * Anamorphic lens model, using Anamorphic lens distortion parameters
+ */
+UCLASS(BlueprintType, meta = (DisplayName = "Anamorphic Lens Model"))
+class CAMERACALIBRATIONCORE_API UAnamorphicLensModel : public ULensModel
+{
+	GENERATED_BODY()
+
+public:
+	//~ Begin ULensModel interface
+ 	virtual UScriptStruct* GetParameterStruct() const override;
+ 	virtual FName GetModelName() const override;
+ 	virtual FName GetShortModelName() const override;
+	//~ End ULensModel interface
+};
diff --git a/Engine/Plugins/VirtualProduction/CameraCalibrationCore/Source/CameraCalibrationCore/Public/Models/LensModel.h b/Engine/Plugins/VirtualProduction/CameraCalibrationCore/Source/CameraCalibrationCore/Public/Models/LensModel.h
--- a/Engine/Plugins/VirtualProduction/CameraCalibrationCore/Source/CameraCalibrationCore/Public/Models/LensModel.h
+++ b/Engine/Plugins/VirtualProduction/CameraCalibrationCore/Source/CameraCalibrationCore/Public/Models/LensModel.h
@@ -30,6 +30,9 @@ public:
 	/** Get the struct of distortion parameters supported by this model */
 	virtual UScriptStruct* GetParameterStruct() const PURE_VIRTUAL(ULensModel::GetParameterStruct, return nullptr;);
 
+	/** Returns an array of floats representing the default values of each property in the parameter struct */
+	void GetDefaultParameterArray(TArray<float>& OutParameterValues) const;
+
 #if WITH_EDITOR
 	/** Get the names of each float parameters supported by this model */
 	virtual TArray<FText> GetParameterDisplayNames() const;
