Camera Calibration: Fix LensFile to broadcast an event when the LensModel changes. Properly update lens component and lens file editor when model changes. Fix distortion tool to properly filter out algos that do not support the current model. Detail customization for LensComponent to correctly display model parameter names.

diff --git a/Engine/Plugins/VirtualProduction/CameraCalibration/Source/CameraCalibrationEditor/Private/AssetEditor/LensDistortionTool.cpp b/Engine/Plugins/VirtualProduction/CameraCalibration/Source/CameraCalibrationEditor/Private/AssetEditor/LensDistortionTool.cpp
--- a/Engine/Plugins/VirtualProduction/CameraCalibration/Source/CameraCalibrationEditor/Private/AssetEditor/LensDistortionTool.cpp
+++ b/Engine/Plugins/VirtualProduction/CameraCalibration/Source/CameraCalibrationEditor/Private/AssetEditor/LensDistortionTool.cpp
@@ -21,7 +21,6 @@
 #include "ScopedTransaction.h"
 #include "Serialization/JsonWriter.h"
 #include "Serialization/JsonSerializer.h"
-#include "SLensDistortionToolPanel.h"
 #include "UObject/Class.h"
 #include "UObject/Package.h"
 #include "UObject/UObjectIterator.h"
@@ -57,6 +56,39 @@ void ULensDistortionTool::Initialize(TWeakPtr<FCameraCalibrationStepsController>
 			}
 		}
 	}
+
+	if (TSharedPtr<FCameraCalibrationStepsController> SharedStepsController = CameraCalibrationStepsController.Pin())
+	{
+		if (ULensFile* const LensFile = SharedStepsController->GetLensFile())
+		{
+			LensFile->OnLensFileModelChanged().AddUObject(this, &ULensDistortionTool::OnLensModelChanged);
+
+			UpdateAlgoMap(LensFile->LensInfo.LensModel);
+		}
+	}
+}
+
+void ULensDistortionTool::OnLensModelChanged(const TSubclassOf<ULensModel>& LensModel)
+{
+	UpdateAlgoMap(LensModel);
+}
+
+void ULensDistortionTool::UpdateAlgoMap(const TSubclassOf<ULensModel>& LensModel)
+{
+	SupportedAlgosMap.Empty();
+	for (const TPair<FName, TSubclassOf<UCameraLensDistortionAlgo>>& AlgoPair : AlgosMap)
+	{
+		const UCameraLensDistortionAlgo* Algo = CastChecked<UCameraLensDistortionAlgo>(AlgoPair.Value->GetDefaultObject());
+
+		if (Algo->SupportsModel(LensModel))
+		{
+			SupportedAlgosMap.Add(AlgoPair);
+		}
+	}
+	if (DistortionWidget)
+	{
+		DistortionWidget->UpdateAlgosOptions();
+	}
 }
 
 void ULensDistortionTool::Shutdown()
@@ -68,6 +100,14 @@ void ULensDistortionTool::Shutdown()
 
 		EndCalibrationSession();
 	}
+
+	if (TSharedPtr<FCameraCalibrationStepsController> SharedStepsController = CameraCalibrationStepsController.Pin())
+	{
+		if (ULensFile* const LensFile = SharedStepsController->GetLensFile())
+		{
+			LensFile->OnLensFileModelChanged().RemoveAll(this);
+		}
+	}
 }
 
 void ULensDistortionTool::Tick(float DeltaTime)
@@ -95,7 +135,8 @@ bool ULensDistortionTool::OnViewportClicked(const FGeometry& MyGeometry, const F
 
 TSharedRef<SWidget> ULensDistortionTool::BuildUI()
 {
-	return SNew(SLensDistortionToolPanel, this);
+	DistortionWidget = SNew(SLensDistortionToolPanel, this);
+	return DistortionWidget.ToSharedRef();
 }
 
 bool ULensDistortionTool::DependsOnStep(UCameraCalibrationStep* Step) const
@@ -119,6 +160,25 @@ void ULensDistortionTool::Deactivate()
 	bIsActive = false;
 }
 
+void ULensDistortionTool::ResetAlgo()
+{
+	// Remove old Algo
+	if (CurrentAlgo)
+	{
+		CurrentAlgo->Shutdown();
+		CurrentAlgo = nullptr;
+
+		EndCalibrationSession();
+	}
+
+	// Set the tool overlay pass' material to the MID associate with the current algo
+	if (TSharedPtr<FCameraCalibrationStepsController> StepsController = CameraCalibrationStepsController.Pin())
+	{
+		StepsController->SetOverlayEnabled(false);
+		StepsController->SetOverlayMaterial(nullptr);
+	}
+}
+
 void ULensDistortionTool::SetAlgo(const FName& AlgoName)
 {
 	// Find the algo class
@@ -273,7 +333,7 @@ FCameraCalibrationStepsController* ULensDistortionTool::GetCameraCalibrationStep
 TArray<FName> ULensDistortionTool::GetAlgos() const
 {
 	TArray<FName> OutKeys;
-	AlgosMap.GetKeys(OutKeys);
+	SupportedAlgosMap.GetKeys(OutKeys);
 	return OutKeys;
 }
 
diff --git a/Engine/Plugins/VirtualProduction/CameraCalibration/Source/CameraCalibrationEditor/Private/AssetEditor/LensDistortionTool.h b/Engine/Plugins/VirtualProduction/CameraCalibration/Source/CameraCalibrationEditor/Private/AssetEditor/LensDistortionTool.h
--- a/Engine/Plugins/VirtualProduction/CameraCalibration/Source/CameraCalibrationEditor/Private/AssetEditor/LensDistortionTool.h
+++ b/Engine/Plugins/VirtualProduction/CameraCalibration/Source/CameraCalibrationEditor/Private/AssetEditor/LensDistortionTool.h
@@ -7,6 +7,7 @@
 #include "Engine/World.h"
 #include "CameraCalibrationStep.h"
 #include "ImageCore.h"
+#include "SLensDistortionToolPanel.h"
 
 #include "LensDistortionTool.generated.h"
 
@@ -15,6 +16,7 @@ struct FPointerEvent;
 
 class FJsonObject;
 class UCameraLensDistortionAlgo;
+class ULensModel;
 
 /** Data associated with a lens distortion calibration session */
 struct FLensDistortionSessionInfo
@@ -58,6 +60,9 @@ public:
 	/** Selects the algorithm by name */
 	void SetAlgo(const FName& AlgoName);
 
+	/** Resets the current algo to none */
+	void ResetAlgo();
+
 	/** Returns the currently selected algorithm */
 	UCameraLensDistortionAlgo* GetAlgo() const;
 
@@ -67,6 +72,9 @@ public:
 	/** Called by the UI when the user wants to save the calibration data that the current algorithm is providing */
 	void OnSaveCurrentCalibrationData();
 
+	/** Triggered when the Lens Model used by the LensFile changes, allowing the tool to update its list of supported algos */
+	void OnLensModelChanged(const TSubclassOf<ULensModel>& LensModel);
+
 	/** Initiate a new calibration session (if one is not already active) */
 	void StartCalibrationSession();
 
@@ -95,6 +103,9 @@ private:
 	/** Get the filename of a row with the input index */
 	FString GetRowFilename(int32 RowIndex) const;
 
+	/** Update the list of algos that support the current Lens Model */
+	void UpdateAlgoMap(const TSubclassOf<ULensModel>& LensModel);
+
 public:
 	/** Stores info about the current calibration session of this tool */
 	FLensDistortionSessionInfo SessionInfo;
@@ -112,10 +123,17 @@ private:
 	UPROPERTY(Transient)
 	TMap<FName, TSubclassOf<UCameraLensDistortionAlgo>> AlgosMap;
 
+	/** Holds a subset of the registered algos that support the current Lens Model */
+	UPROPERTY(Transient)
+	TMap<FName, TSubclassOf<UCameraLensDistortionAlgo>> SupportedAlgosMap;
+
 	/** Map of algo names to overlay MIDs used by those algos */
 	UPROPERTY(Transient)
 	TMap<FName, TObjectPtr<UMaterialInstanceDynamic>> AlgoOverlayMIDs;
 
 	/** True if this tool is the active one in the panel */
 	bool bIsActive = false;
+
+	/** UI Widget for this Tool */
+	TSharedPtr<SLensDistortionToolPanel> DistortionWidget;
 };
diff --git a/Engine/Plugins/VirtualProduction/CameraCalibration/Source/CameraCalibrationEditor/Private/AssetEditor/LensInfoStep.cpp b/Engine/Plugins/VirtualProduction/CameraCalibration/Source/CameraCalibrationEditor/Private/AssetEditor/LensInfoStep.cpp
--- a/Engine/Plugins/VirtualProduction/CameraCalibration/Source/CameraCalibrationEditor/Private/AssetEditor/LensInfoStep.cpp
+++ b/Engine/Plugins/VirtualProduction/CameraCalibration/Source/CameraCalibrationEditor/Private/AssetEditor/LensInfoStep.cpp
@@ -38,14 +38,7 @@ TSharedRef<SWidget> ULensInfoStep::BuildUI()
 	TSharedRef<FStructOnScope> LensInfoStructOnScope = MakeShared<FStructOnScope>(FLensInfo::StaticStruct(), reinterpret_cast<uint8*>(&LensFile->LensInfo));
 	TSharedPtr<IStructureDetailsView> LensInfoStructureDetailsView = PropertyEditor.CreateStructureDetailView(DetailArgs, LensInfoStructDetailsView, LensInfoStructOnScope);
 
-	LensInfoStructureDetailsView->GetOnFinishedChangingPropertiesDelegate().AddLambda([&](const FPropertyChangedEvent& PropertyChangedEvent)
-	{
-		// Ignore temporary interaction (dragging sliders, etc.)
-		if (PropertyChangedEvent.ChangeType == EPropertyChangeType::ValueSet)
-		{
-			OnSaveLensInformation();
-		}
-	});
+	LensInfoStructureDetailsView->GetDetailsView()->OnFinishedChangingProperties().AddUObject(this, &ULensInfoStep::OnLensInfoChanged);
 
 	/** Camera Feed Info Details View */
 	FStructureDetailsViewArgs CameraFeedInfoStructDetailsView;
@@ -152,14 +145,39 @@ void ULensInfoStep::Deactivate()
 	bIsActive = false;
 }
 
-void ULensInfoStep::OnSaveLensInformation()
+void ULensInfoStep::OnLensInfoChanged(const FPropertyChangedEvent& PropertyChangedEvent)
+{
+	// Ignore temporary interaction (dragging sliders, etc.)
+	if (PropertyChangedEvent.ChangeType != EPropertyChangeType::ValueSet)
+	{
+		return;
+	}
+
+	const FName PropertyName = PropertyChangedEvent.Property->GetFName();
+
+	if (PropertyName == GET_MEMBER_NAME_CHECKED(FLensInfo, LensModel))
+	{
+		if (ULensFile* const LensFile = CameraCalibrationStepsController.Pin()->GetLensFile())
+		{
+			LensFile->OnLensFileModelChanged().Broadcast(LensFile->LensInfo.LensModel);
+		}
+	}
+
+	SaveLensInformation();
+}
+
+void ULensInfoStep::SaveLensInformation()
 {
 	if (!CameraCalibrationStepsController.IsValid())
 	{
 		return;
 	}
 
-	ULensFile* LensFile = CameraCalibrationStepsController.Pin()->GetLensFile();
+	ULensFile* const LensFile = CameraCalibrationStepsController.Pin()->GetLensFile();
+	if (!LensFile)
+	{
+		return;
+	}
 
 	// Validate sensor dimensions
 	constexpr float MinimumSize = 1.0f; // Limit sensor dimension to 1mm
@@ -305,7 +323,7 @@ void ULensInfoStep::ResetToDefault()
 {
 	ULensFile* LensFile = CameraCalibrationStepsController.Pin()->GetLensFile();
 	LensFile->LensInfo = OriginalLensInfo;
-	OnSaveLensInformation();
+	SaveLensInformation();
 }
 
 bool ULensInfoStep::DiffersFromDefault() const
diff --git a/Engine/Plugins/VirtualProduction/CameraCalibration/Source/CameraCalibrationEditor/Private/AssetEditor/LensInfoStep.h b/Engine/Plugins/VirtualProduction/CameraCalibration/Source/CameraCalibrationEditor/Private/AssetEditor/LensInfoStep.h
--- a/Engine/Plugins/VirtualProduction/CameraCalibration/Source/CameraCalibrationEditor/Private/AssetEditor/LensInfoStep.h
+++ b/Engine/Plugins/VirtualProduction/CameraCalibration/Source/CameraCalibrationEditor/Private/AssetEditor/LensInfoStep.h
@@ -33,11 +33,14 @@ public:
 public:
 
 	/** Triggered when user edits any of the LensInfo properties */
-	void OnSaveLensInformation();
+	void OnLensInfoChanged(const FPropertyChangedEvent& PropertyChangedEvent);
 
 	/** Triggered when user edits any of the CameraFeedInfo properties */
 	void OnCameraFeedInfoChanged(const FPropertyChangedEvent& PropertyChangedEvent);
 
+	/** Validates the new LensInfo settings and records a transaction for them */
+	void SaveLensInformation();
+
 	/** Reapply the original LensInfo to the LensFile */
 	void ResetToDefault();
 
@@ -50,7 +53,6 @@ public:
 	/** Returns true if the camera feed info has been overriden by the user */
 	bool CameraFeedInfoDiffersFromDefault() const;
 
-
 private:
 	/** Determines the visibility of the aspect ratio warning by comparing the camera feed aspect ratio with that of the camera actor */
 	EVisibility HandleAspectRatioWarningVisibility() const;
diff --git a/Engine/Plugins/VirtualProduction/CameraCalibration/Source/CameraCalibrationEditor/Private/AssetEditor/SLensDistortionToolPanel.cpp b/Engine/Plugins/VirtualProduction/CameraCalibration/Source/CameraCalibrationEditor/Private/AssetEditor/SLensDistortionToolPanel.cpp
--- a/Engine/Plugins/VirtualProduction/CameraCalibration/Source/CameraCalibrationEditor/Private/AssetEditor/SLensDistortionToolPanel.cpp
+++ b/Engine/Plugins/VirtualProduction/CameraCalibration/Source/CameraCalibrationEditor/Private/AssetEditor/SLensDistortionToolPanel.cpp
@@ -132,19 +132,20 @@ void SLensDistortionToolPanel::UpdateUI()
 	const UCameraLensDistortionAlgo* OldAlgo = Tool->GetAlgo();
 
 	// Set new algo by name
-	const FName AlgoName(*AlgosComboBox->GetSelectedItem());
-	Tool->SetAlgo(AlgoName);
+	TSharedPtr<FString> SelectedAlgo = AlgosComboBox->GetSelectedItem();
+	if (!SelectedAlgo)
+	{
+		Tool->ResetAlgo();
+	}
+	else
+	{
+		const FName AlgoName(*SelectedAlgo);
+		Tool->SetAlgo(AlgoName);
+	}
 
 	// Get the new algo
 	UCameraLensDistortionAlgo* Algo = Tool->GetAlgo();
 
-	// nullptr may indicate that it was unregistered, so refresh combobox options.
-	if (!Algo)
-	{
-		UpdateAlgosOptions();
-		return;
-	}
-
 	// If we didn't change the algo, we're done here.
 	if (Algo == OldAlgo)
 	{
@@ -155,6 +156,13 @@ void SLensDistortionToolPanel::UpdateUI()
 	check(UI.IsValid());
 	UI->ClearChildren();
 
+	// nullptr may indicate that it was unregistered, so refresh combobox options.
+	if (!Algo)
+	{
+		UpdateAlgosOptions();
+		return;
+	}
+
 	// Assign GUI
 	UI->AddSlot() [Algo->BuildUI()];
 }
@@ -175,6 +183,16 @@ void SLensDistortionToolPanel::UpdateAlgosOptions()
 		CurrentAlgos.Add(MakeShared<FString>(AlgoName.ToString()));
 	}
 
+	// Pick the first option by default (if available)
+	if (CurrentAlgos.Num())
+	{
+		AlgosComboBox->SetSelectedItem(CurrentAlgos[0]);
+	}
+	else
+	{
+		AlgosComboBox->SetSelectedItem(nullptr);
+	}
+
 	// Ask the ComboBox to refresh its options from its source (that we just updated)
 	AlgosComboBox->RefreshOptions();
 }
diff --git a/Engine/Plugins/VirtualProduction/CameraCalibration/Source/CameraCalibrationEditor/Private/AssetEditor/SLensEvaluation.cpp b/Engine/Plugins/VirtualProduction/CameraCalibration/Source/CameraCalibrationEditor/Private/AssetEditor/SLensEvaluation.cpp
--- a/Engine/Plugins/VirtualProduction/CameraCalibration/Source/CameraCalibrationEditor/Private/AssetEditor/SLensEvaluation.cpp
+++ b/Engine/Plugins/VirtualProduction/CameraCalibration/Source/CameraCalibrationEditor/Private/AssetEditor/SLensEvaluation.cpp
@@ -23,6 +23,12 @@ void SLensEvaluation::Construct(const FArguments& InArgs, TWeakPtr<FCameraCalibr
 
 	WeakStepsController = StepsController;
 
+	if (InLensFile)
+	{
+		MakeDistortionWidget(InLensFile->LensInfo.LensModel);
+		InLensFile->OnLensFileModelChanged().AddSP(this, &SLensEvaluation::MakeDistortionWidget);
+	}
+
 	ChildSlot
 	[
 		SNew(SBorder)
@@ -59,7 +65,7 @@ void SLensEvaluation::Construct(const FArguments& InArgs, TWeakPtr<FCameraCalibr
 				.Padding(5.0f, 5.0f)
 				.FillWidth(0.2f)
 				[
-					MakeDistortionWidget()
+					DistortionWidget.ToSharedRef()
 				]
 				//Image Center section
 				+ SHorizontalBox::Slot()
@@ -80,6 +86,14 @@ void SLensEvaluation::Construct(const FArguments& InArgs, TWeakPtr<FCameraCalibr
 	];
 }
 
+SLensEvaluation::~SLensEvaluation()
+{
+	if (ULensFile* LensFilePtr = LensFile.Get())
+	{
+		LensFilePtr->OnLensFileModelChanged().RemoveAll(this);
+	}
+}
+
 void SLensEvaluation::Tick(const FGeometry& AllottedGeometry, const double InCurrentTime, const float InDeltaTime)
 {
 	//Cache LensFile evaluation data every tick to be sure we have the right one for the frame during calibration
@@ -493,8 +507,15 @@ TSharedRef<SWidget> SLensEvaluation::MakeEvaluatedFIZWidget() const
 		];
 }
 
-TSharedRef<SWidget> SLensEvaluation::MakeDistortionWidget() const
+void SLensEvaluation::MakeDistortionWidget(const TSubclassOf<ULensModel>& LensModel)
 {
+	if (!DistortionWidget)
+	{
+		DistortionWidget = SNew(SVerticalBox);
+	}
+
+	DistortionWidget->ClearChildren();
+
 	//Find the named distortion parameters the current model has
 	TArray<FText> Parameters;
 	if (LensFile->LensInfo.LensModel)
@@ -508,18 +529,19 @@ TSharedRef<SWidget> SLensEvaluation::MakeDistortionWidget() const
 		.Font(FAppStyle::GetFontStyle("DetailsView.CategoryFontStyle"))
 		.ShadowOffset(FVector2D(1.0f, 1.0f));
 
+	DistortionWidget->AddSlot()
+		.Padding(0.0f, 5.0f)
+		.AutoHeight()
+		.HAlign(HAlign_Left)
+		[
+			Title
+		]
+	;
+
 	//if there are no parameters, create a simpler widget
 	if (Parameters.Num() <= 0)
 	{
-		return SNew(SVerticalBox)
-			+ SVerticalBox::Slot()
-			.Padding(0.0f, 5.0f)
-			.AutoHeight()
-			.HAlign(HAlign_Left)
-			[
-				Title
-			]
-			+ SVerticalBox::Slot()
+		DistortionWidget->AddSlot()
 			.Padding(0.0f, 5.0f)
 			.AutoHeight()
 			.HAlign(HAlign_Left)
@@ -565,15 +587,7 @@ TSharedRef<SWidget> SLensEvaluation::MakeDistortionWidget() const
 			];
 	}
 
-	return SNew(SVerticalBox)
-		+ SVerticalBox::Slot()
-		.Padding(0.0f, 5.0f)
-		.AutoHeight()
-		.HAlign(HAlign_Left)
-		[
-			Title
-		]
-		+ SVerticalBox::Slot()
+	DistortionWidget->AddSlot()
 		.AutoHeight()
 		.HAlign(HAlign_Left)
 		[
diff --git a/Engine/Plugins/VirtualProduction/CameraCalibration/Source/CameraCalibrationEditor/Private/AssetEditor/SLensEvaluation.h b/Engine/Plugins/VirtualProduction/CameraCalibration/Source/CameraCalibrationEditor/Private/AssetEditor/SLensEvaluation.h
--- a/Engine/Plugins/VirtualProduction/CameraCalibration/Source/CameraCalibrationEditor/Private/AssetEditor/SLensEvaluation.h
+++ b/Engine/Plugins/VirtualProduction/CameraCalibration/Source/CameraCalibrationEditor/Private/AssetEditor/SLensEvaluation.h
@@ -10,8 +10,10 @@
 #include "Widgets/DeclarativeSyntaxSupport.h"
 
 
-class ULensFile;
 class FCameraCalibrationStepsController;
+class SVerticalBox;
+class ULensFile;
+class ULensModel;
 
 /**
  * Widget using lens file evaluation inputs to evaluate lens file and show resulting data
@@ -24,6 +26,8 @@ public:
 	{}
 	SLATE_END_ARGS()
 
+	~SLensEvaluation();
+
 	void Construct(const FArguments& InArgs, TWeakPtr<FCameraCalibrationStepsController> StepsController, ULensFile* InLensFile);
 
 	//~ Begin SCompoundWidget interface
@@ -63,7 +67,7 @@ private:
 	TSharedRef<SWidget> MakeEvaluatedFIZWidget() const;
 
 	/** Make widget showing distortion data evaluated from LensFile */
-	TSharedRef<SWidget> MakeDistortionWidget() const;
+	void MakeDistortionWidget(const TSubclassOf<ULensModel>& LensModel);
 
 	/** Make widget showing Intrinsic data evaluated from LensFile */
 	TSharedRef<SWidget> MakeIntrinsicsWidget() const;
@@ -85,10 +89,13 @@ private:
 	/** Calibration steps controller */
 	TWeakPtr<FCameraCalibrationStepsController> WeakStepsController;
 
-	//~ Cached LensFile evaluated data
+	/** Cached LensFile evaluated data */
 	FCachedFIZData CachedFIZData;
 	FDistortionInfo CachedDistortionInfo;
 	FFocalLengthInfo CachedFocalLengthInfo;
 	FImageCenterInfo CachedImageCenter;
 	FNodalPointOffset CachedNodalOffset;
-};
\ No newline at end of file
+
+	/** Distortion Widget that displays the evaluated distortion parameters */
+	TSharedPtr<SVerticalBox> DistortionWidget;
+};
diff --git a/Engine/Plugins/VirtualProduction/CameraCalibration/Source/CameraCalibrationEditor/Private/Calibrators/CameraLensDistortionAlgoCheckerboard.cpp b/Engine/Plugins/VirtualProduction/CameraCalibration/Source/CameraCalibrationEditor/Private/Calibrators/CameraLensDistortionAlgoCheckerboard.cpp
--- a/Engine/Plugins/VirtualProduction/CameraCalibration/Source/CameraCalibrationEditor/Private/Calibrators/CameraLensDistortionAlgoCheckerboard.cpp
+++ b/Engine/Plugins/VirtualProduction/CameraCalibration/Source/CameraCalibrationEditor/Private/Calibrators/CameraLensDistortionAlgoCheckerboard.cpp
@@ -137,6 +137,11 @@ void UCameraLensDistortionAlgoCheckerboard::Shutdown()
 	CalibrationRows.Reset();
 }
 
+bool UCameraLensDistortionAlgoCheckerboard::SupportsModel(const TSubclassOf<ULensModel>& LensModel) const
+{
+	return (LensModel == USphericalLensModel::StaticClass());
+}
+
 void UCameraLensDistortionAlgoCheckerboard::Tick(float DeltaTime)
 {
 	if (!ensure(Tool.IsValid()))
diff --git a/Engine/Plugins/VirtualProduction/CameraCalibration/Source/CameraCalibrationEditor/Private/Calibrators/CameraLensDistortionAlgoCheckerboard.h b/Engine/Plugins/VirtualProduction/CameraCalibration/Source/CameraCalibrationEditor/Private/Calibrators/CameraLensDistortionAlgoCheckerboard.h
--- a/Engine/Plugins/VirtualProduction/CameraCalibration/Source/CameraCalibrationEditor/Private/Calibrators/CameraLensDistortionAlgoCheckerboard.h
+++ b/Engine/Plugins/VirtualProduction/CameraCalibration/Source/CameraCalibrationEditor/Private/Calibrators/CameraLensDistortionAlgoCheckerboard.h
@@ -130,6 +130,7 @@ public:
 	//~ Begin CameraLensDistortionAlgo
 	virtual void Initialize(ULensDistortionTool* InTool) override;
 	virtual void Shutdown() override;
+	virtual bool SupportsModel(const TSubclassOf<ULensModel>& LensModel) const override;
 	virtual void Tick(float DeltaTime) override;
 	virtual bool OnViewportClicked(const FGeometry& MyGeometry, const FPointerEvent& MouseEvent) override;
 	virtual TSharedRef<SWidget> BuildUI() override;
diff --git a/Engine/Plugins/VirtualProduction/CameraCalibration/Source/CameraCalibrationEditor/Private/Calibrators/CameraLensDistortionAlgoPoints.cpp b/Engine/Plugins/VirtualProduction/CameraCalibration/Source/CameraCalibrationEditor/Private/Calibrators/CameraLensDistortionAlgoPoints.cpp
--- a/Engine/Plugins/VirtualProduction/CameraCalibration/Source/CameraCalibrationEditor/Private/Calibrators/CameraLensDistortionAlgoPoints.cpp
+++ b/Engine/Plugins/VirtualProduction/CameraCalibration/Source/CameraCalibrationEditor/Private/Calibrators/CameraLensDistortionAlgoPoints.cpp
@@ -124,6 +124,11 @@ void UCameraLensDistortionAlgoPoints::Shutdown()
 	LensDistortionTool.Reset();
 }
 
+bool UCameraLensDistortionAlgoPoints::SupportsModel(const TSubclassOf<ULensModel>& LensModel) const
+{
+	return (LensModel == USphericalLensModel::StaticClass());
+}
+
 void UCameraLensDistortionAlgoPoints::Tick(float DeltaTime)
 {
 	if (!LensDistortionTool.IsValid())
diff --git a/Engine/Plugins/VirtualProduction/CameraCalibration/Source/CameraCalibrationEditor/Private/Calibrators/CameraLensDistortionAlgoPoints.h b/Engine/Plugins/VirtualProduction/CameraCalibration/Source/CameraCalibrationEditor/Private/Calibrators/CameraLensDistortionAlgoPoints.h
--- a/Engine/Plugins/VirtualProduction/CameraCalibration/Source/CameraCalibrationEditor/Private/Calibrators/CameraLensDistortionAlgoPoints.h
+++ b/Engine/Plugins/VirtualProduction/CameraCalibration/Source/CameraCalibrationEditor/Private/Calibrators/CameraLensDistortionAlgoPoints.h
@@ -111,6 +111,7 @@ public:
 	//~ Begin UCameraLensDistortionAlgo
 	virtual void Initialize(ULensDistortionTool* InLensDistortionTool) override;
 	virtual void Shutdown() override;
+	virtual bool SupportsModel(const TSubclassOf<ULensModel>& LensModel) const override;
 	virtual void Tick(float DeltaTime) override;
 	virtual bool OnViewportClicked(const FGeometry& MyGeometry, const FPointerEvent& MouseEvent) override;
 	virtual TSharedRef<SWidget> BuildUI() override;
diff --git a/Engine/Plugins/VirtualProduction/CameraCalibrationCore/Source/CameraCalibrationCore/Private/LensComponent.cpp b/Engine/Plugins/VirtualProduction/CameraCalibrationCore/Source/CameraCalibrationCore/Private/LensComponent.cpp
--- a/Engine/Plugins/VirtualProduction/CameraCalibrationCore/Source/CameraCalibrationCore/Private/LensComponent.cpp
+++ b/Engine/Plugins/VirtualProduction/CameraCalibrationCore/Source/CameraCalibrationCore/Private/LensComponent.cpp
@@ -258,6 +258,17 @@ void ULensComponent::PostLoad()
 		bIsDistortionSetup = false;
 	}
 
+	// Cache the currently set LensFile and register for new LensFile events
+	WeakCachedLensFile = LensFilePicker.GetLensFile();
+
+	if (ULensFile* const LensFile = WeakCachedLensFile.Get())
+	{
+		if (!LensFile->OnLensFileModelChanged().IsBoundToObject(this))
+		{
+			LensFile->OnLensFileModelChanged().AddUObject(this, &ULensComponent::OnLensFileModelChanged);
+		}
+	}
+
 #if WITH_EDITOR
 	const int32 UE5MainVersion = GetLinkerCustomVersion(FUE5MainStreamObjectVersion::GUID);
 	if (UE5MainVersion < FUE5MainStreamObjectVersion::LensComponentNodalOffset)
@@ -474,6 +485,20 @@ void ULensComponent::SetLensFilePicker(FLensFilePicker LensFile)
 {
 	LensFilePicker = LensFile;
 
+	if (ULensFile* const CachedLensFile = WeakCachedLensFile.Get())
+	{
+		CachedLensFile->OnLensFileModelChanged().RemoveAll(this);
+	}
+
+	WeakCachedLensFile = LensFilePicker.GetLensFile();
+	if (ULensFile* const NewLensFile = WeakCachedLensFile.Get())
+	{
+		if (!NewLensFile->OnLensFileModelChanged().IsBoundToObject(this))
+		{
+			NewLensFile->OnLensFileModelChanged().AddUObject(this, &ULensComponent::OnLensFileModelChanged);
+		}
+	}
+
 	if (DistortionStateSource == EDistortionSource::LensFile)
 	{
 		if (ULensFile* const NewLensFile = LensFilePicker.GetLensFile())
@@ -490,20 +515,11 @@ void ULensComponent::SetLensFilePicker(FLensFilePicker LensFile)
 void ULensComponent::SetLensFile(ULensFile* Lens)
 {
 	// Automatically sets this to false so the component can use the newly set LensFile directly
-	LensFilePicker.bUseDefaultLensFile = false;
-	LensFilePicker.LensFile = Lens;
+	FLensFilePicker NewLensFilePicker;
+	NewLensFilePicker.bUseDefaultLensFile = false;
+	NewLensFilePicker.LensFile = Lens;
 
-	if (DistortionStateSource == EDistortionSource::LensFile)
-	{
-		if (Lens)
-		{
-			SetLensModel(Lens->LensInfo.LensModel);
-		}
-		else
-		{
-			ClearDistortionState();
-		}
-	}
+	SetLensFilePicker(NewLensFilePicker);
 }
 
 EFIZEvaluationMode ULensComponent::GetFIZEvaluationMode() const
@@ -600,6 +616,16 @@ void ULensComponent::SetLensModel(TSubclassOf<ULensModel> Model)
 	{
 		DistortionState.DistortionInfo.Parameters.Empty();
 	}
+
+	OnLensComponentModelChangedDelegate.Broadcast(LensModel);
+}
+
+void ULensComponent::OnLensFileModelChanged(const TSubclassOf<ULensModel>& Model)
+{
+	if (DistortionStateSource == EDistortionSource::LensFile)
+	{
+		SetLensModel(Model);
+	}
 }
 
 FLensDistortionState ULensComponent::GetDistortionState() const
@@ -621,16 +647,16 @@ void ULensComponent::SetDistortionSource(EDistortionSource Source)
 {
 	DistortionStateSource = Source;
 
-	ClearDistortionState();
-
 	// If the new source is a LensFile, update the lens model to match
 	if (DistortionStateSource == EDistortionSource::LensFile)
 	{
 		if (ULensFile* LensFile = LensFilePicker.GetLensFile())
 		{
-			LensModel = LensFile->LensInfo.LensModel;
+			SetLensModel(LensFile->LensInfo.LensModel);
 		}
 	}
+
+	ClearDistortionState();
 }
 
 bool ULensComponent::ShouldApplyDistortion() const
diff --git a/Engine/Plugins/VirtualProduction/CameraCalibrationCore/Source/CameraCalibrationCore/Private/LensFile.cpp b/Engine/Plugins/VirtualProduction/CameraCalibrationCore/Source/CameraCalibrationCore/Private/LensFile.cpp
--- a/Engine/Plugins/VirtualProduction/CameraCalibrationCore/Source/CameraCalibrationCore/Private/LensFile.cpp
+++ b/Engine/Plugins/VirtualProduction/CameraCalibrationCore/Source/CameraCalibrationCore/Private/LensFile.cpp
@@ -217,6 +217,8 @@ void ULensFile::PostEditChangeChainProperty(struct FPropertyChangedChainEvent& P
 		}
 		else if (PropertyName == GET_MEMBER_NAME_CHECKED(FLensInfo, LensModel))
 		{
+			OnLensFileModelChangedDelegate.Broadcast(LensInfo.LensModel);
+
 			//LensModel has changed, clear distortion and focal length tables
 			LensDataTableUtils::EmptyTable(DistortionTable);
 			LensDataTableUtils::EmptyTable(FocalLengthTable);
diff --git a/Engine/Plugins/VirtualProduction/CameraCalibrationCore/Source/CameraCalibrationCore/Public/CameraLensDistortionAlgo.h b/Engine/Plugins/VirtualProduction/CameraCalibrationCore/Source/CameraCalibrationCore/Public/CameraLensDistortionAlgo.h
--- a/Engine/Plugins/VirtualProduction/CameraCalibrationCore/Source/CameraCalibrationCore/Public/CameraLensDistortionAlgo.h
+++ b/Engine/Plugins/VirtualProduction/CameraCalibrationCore/Source/CameraCalibrationCore/Public/CameraLensDistortionAlgo.h
@@ -42,6 +42,9 @@ public:
 	/** Called every frame */
 	virtual void Tick(float DeltaTime) {};
 
+	/** Returns true if the algo supports the input Lens Model */
+	virtual bool SupportsModel(const TSubclassOf<ULensModel>& LensModel) const { return false; };
+
 	/** Callback when viewport is clicked. Returns false if the event was not handled. */
 	virtual bool OnViewportClicked(const FGeometry& MyGeometry, const FPointerEvent& MouseEvent) { return false;  };
 
diff --git a/Engine/Plugins/VirtualProduction/CameraCalibrationCore/Source/CameraCalibrationCore/Public/LensComponent.h b/Engine/Plugins/VirtualProduction/CameraCalibrationCore/Source/CameraCalibrationCore/Public/LensComponent.h
--- a/Engine/Plugins/VirtualProduction/CameraCalibrationCore/Source/CameraCalibrationCore/Public/LensComponent.h
+++ b/Engine/Plugins/VirtualProduction/CameraCalibrationCore/Source/CameraCalibrationCore/Public/LensComponent.h
@@ -10,8 +10,13 @@
 
 #include "LensComponent.generated.h"
 
+class FLensComponentDetailCustomization;
 class UCineCameraComponent;
 
+
+DECLARE_MULTICAST_DELEGATE_OneParam(FOnLensComponentModelChanged, const TSubclassOf<ULensModel>&);
+
+
 /** Mode that controls where FIZ inputs are sourced from and how they are used to evaluate the LensFile */
 UENUM(BlueprintType)
 enum class EFIZEvaluationMode : uint8
@@ -56,6 +61,8 @@ class CAMERACALIBRATIONCORE_API ULensComponent : public UActorComponent
 {
 	GENERATED_BODY()
 
+	friend FLensComponentDetailCustomization;
+
 public:
 	ULensComponent();
 
@@ -76,6 +83,12 @@ public:
 #endif //WITH_EDITOR
 	//~ End UObject interface
 
+	/** Returns the delegate that is triggered when the LensModel changes */
+	FOnLensComponentModelChanged& OnLensComponentModelChanged() 
+	{ 
+		return OnLensComponentModelChangedDelegate; 
+	}
+
 	/** Get the LensFile picker used by this component */
 	UFUNCTION(BlueprintPure, Category = "Lens Component")
 	FLensFilePicker GetLensFilePicker() const;
@@ -227,6 +240,9 @@ private:
 	/** Register to the new LiveLink component's callback to be notified when its controller map changes */
 	void OnLiveLinkComponentRegistered(ULiveLinkComponentController* LiveLinkComponent);
 	
+	/** Triggered when the LensFile model changes */
+	void OnLensFileModelChanged(const TSubclassOf<ULensModel>& Model);
+
 	/** Callback executed when a LiveLink component on the same actor ticks */
 	void ProcessLiveLinkData(const ULiveLinkComponentController* const LiveLinkComponent, const FLiveLinkSubjectFrameData& SubjectData);
 
@@ -272,7 +288,7 @@ protected:
 	TSubclassOf<ULensModel> LensModel;
 
 	/** The current distortion state */
-	UPROPERTY(Interp, EditAnywhere, Category = "Distortion", meta = (ShowOnlyInnerProperties))
+	UPROPERTY(Interp, EditAnywhere, Category = "Distortion", meta = (EditCondition = "DistortionStateSource == EDistortionSource::Manual"))
 	FLensDistortionState DistortionState;
 
 	/** Whether to scale the computed overscan by the overscan percentage */
@@ -330,6 +346,9 @@ protected:
 	UPROPERTY()
 	FString TrackedComponentName;
 
+	/** Weak pointer to the currently set LensFile. Useful to unregister events from the previous LensFile when the asset is swapped out */
+	TWeakObjectPtr<ULensFile> WeakCachedLensFile;
+
 #if WITH_EDITORONLY_DATA
 	UE_DEPRECATED(5.1, "This property has been deprecated. The LensDistortion component no longer tracks the attached camera's original rotation.")
 	UPROPERTY()
@@ -375,4 +394,7 @@ private:
 
 	/** Latest focal length of the target camera, used to track external changes to the original (no overscan applied) focal length */
 	float LastFocalLength = -1.0f;
+
+	/** Delegate that is triggered when the LensModel changes */
+	FOnLensComponentModelChanged OnLensComponentModelChangedDelegate;
 };
diff --git a/Engine/Plugins/VirtualProduction/CameraCalibrationCore/Source/CameraCalibrationCore/Public/LensFile.h b/Engine/Plugins/VirtualProduction/CameraCalibrationCore/Source/CameraCalibrationCore/Public/LensFile.h
--- a/Engine/Plugins/VirtualProduction/CameraCalibrationCore/Source/CameraCalibrationCore/Public/LensFile.h
+++ b/Engine/Plugins/VirtualProduction/CameraCalibrationCore/Source/CameraCalibrationCore/Public/LensFile.h
@@ -29,6 +29,9 @@ class UTextureRenderTarget2D;
 struct FBaseLensTable;
 
 
+DECLARE_MULTICAST_DELEGATE_OneParam(FOnLensFileModelChanged, const TSubclassOf<ULensModel>&);
+
+
 /** Mode of operation of Lens File */
 UENUM()
 enum class ELensDataMode : uint8
@@ -213,6 +216,9 @@ public:
 	/** Get data table reference based on given category */
 	const FBaseLensTable* GetDataTable(ELensDataCategory InDataCategory) const;
 	
+	/** Returns the delegate that is triggered when the LensModel changes */
+	FOnLensFileModelChanged& OnLensFileModelChanged() { return OnLensFileModelChangedDelegate; }
+
 protected:
 	/** Updates derived data entries to make sure it matches what is assigned in map points based on data mode */
 	void UpdateDerivedData();
@@ -310,6 +316,9 @@ protected:
 
 	/** UV coordinates of 8 points (4 corners + 4 mid points) */
 	static const TArray<FVector2D> UndistortedUVs;
+
+	/** Delegate that is triggered when the LensModel changes */
+	FOnLensFileModelChanged OnLensFileModelChangedDelegate;
 };
 
 
diff --git a/Engine/Plugins/VirtualProduction/CameraCalibrationCore/Source/CameraCalibrationCoreEditor/Private/CameraCalibrationCoreEditorModule.cpp b/Engine/Plugins/VirtualProduction/CameraCalibrationCore/Source/CameraCalibrationCoreEditor/Private/CameraCalibrationCoreEditorModule.cpp
--- a/Engine/Plugins/VirtualProduction/CameraCalibrationCore/Source/CameraCalibrationCoreEditor/Private/CameraCalibrationCoreEditorModule.cpp
+++ b/Engine/Plugins/VirtualProduction/CameraCalibrationCore/Source/CameraCalibrationCoreEditor/Private/CameraCalibrationCoreEditorModule.cpp
@@ -12,6 +12,8 @@
 #include "CameraCalibrationCoreEditorStyle.h"
 #include "Editor.h"
 #include "IPlacementModeModule.h"
+#include "LensComponent.h"
+#include "LensComponentDetailCustomization.h"
 #include "LensFile.h"
 #include "Misc/CoreDelegates.h"
 #include "Modules/ModuleManager.h"
@@ -37,6 +39,11 @@ void FCameraCalibrationCoreEditorModule::StartupModule()
 		FOnGetDetailCustomizationInstance::CreateStatic(&FCalibrationPointComponentDetails::MakeInstance)
 	);
 
+	PropertyEditorModule.RegisterCustomClassLayout(
+		ULensComponent::StaticClass()->GetFName(),
+		FOnGetDetailCustomizationInstance::CreateStatic(&FLensComponentDetailCustomization::MakeInstance)
+	);
+
 	RegisterPlacementModeItems();
 }
 
@@ -48,8 +55,8 @@ void FCameraCalibrationCoreEditorModule::ShutdownModule()
 	{
 		FPropertyEditorModule& PropertyModule = FModuleManager::GetModuleChecked<FPropertyEditorModule>("PropertyEditor");
 
-		PropertyModule.UnregisterCustomClassLayout(ULensFile::StaticClass()->GetFName());
 		PropertyModule.UnregisterCustomClassLayout(UCalibrationPointComponent::StaticClass()->GetFName());
+		PropertyModule.UnregisterCustomClassLayout(ULensComponent::StaticClass()->GetFName());
 
 		UnregisterPlacementModeItems();
 	}
diff --git a/Engine/Plugins/VirtualProduction/CameraCalibrationCore/Source/CameraCalibrationCoreEditor/Private/LensComponentDetailCustomization.cpp b/Engine/Plugins/VirtualProduction/CameraCalibrationCore/Source/CameraCalibrationCoreEditor/Private/LensComponentDetailCustomization.cpp
new file mode 100644
--- /dev/null
+++ b/Engine/Plugins/VirtualProduction/CameraCalibrationCore/Source/CameraCalibrationCoreEditor/Private/LensComponentDetailCustomization.cpp
@@ -0,0 +1,59 @@
+// Copyright Epic Games, Inc. All Rights Reserved.
+
+#include "LensComponentDetailCustomization.h"
+
+#include "DetailLayoutBuilder.h"
+#include "LensDistortionStateDetailCustomization.h"
+
+FLensComponentDetailCustomization::~FLensComponentDetailCustomization()
+{
+	if (ULensComponent* LensComponent = WeakLensComponent.Get())
+	{
+		LensComponent->OnLensComponentModelChanged().RemoveAll(this);
+	}
+}
+
+void FLensComponentDetailCustomization::CustomizeDetails(IDetailLayoutBuilder& DetailBuilder)
+{
+	DetailLayout = &DetailBuilder;
+
+	TArray<TWeakObjectPtr<UObject>> SelectedObjects = DetailBuilder.GetSelectedObjects();
+
+ 	if (SelectedObjects.Num() != 1)
+ 	{
+ 		return;
+ 	}
+
+	if (ULensComponent* CustomizedLensComponent = Cast<ULensComponent>(SelectedObjects[0].Get()))
+	{
+		WeakLensComponent = CustomizedLensComponent;
+
+		if (!CustomizedLensComponent->OnLensComponentModelChanged().IsBoundToObject(this))
+		{
+			CustomizedLensComponent->OnLensComponentModelChanged().AddSP(this, &FLensComponentDetailCustomization::OnLensModelChanged);
+		}
+
+		TSharedRef<IPropertyHandle> LensModelPropertyHandle = DetailBuilder.GetProperty(GET_MEMBER_NAME_CHECKED(ULensComponent, LensModel));
+		LensModelPropertyHandle->SetOnPropertyValueChanged(FSimpleDelegate::CreateSP(this, &FLensComponentDetailCustomization::OnLensModelChanged));
+
+		const TSubclassOf<ULensModel> LensModel = CustomizedLensComponent->GetLensModel();
+
+		FOnGetPropertyTypeCustomizationInstance DistortionStateCustomizationDelegate = FOnGetPropertyTypeCustomizationInstance::CreateLambda(
+			[LensModel]
+			{
+				return MakeShared<FLensDistortionStateDetailCustomization>(LensModel);
+			});
+
+		DetailBuilder.RegisterInstancedCustomPropertyTypeLayout(FLensDistortionState::StaticStruct()->GetFName(), DistortionStateCustomizationDelegate);
+	}
+}
+
+void FLensComponentDetailCustomization::OnLensModelChanged()
+{
+	DetailLayout->ForceRefreshDetails();
+}
+
+void FLensComponentDetailCustomization::OnLensModelChanged(const TSubclassOf<ULensModel>& LensModel)
+{
+	DetailLayout->ForceRefreshDetails();
+}
diff --git a/Engine/Plugins/VirtualProduction/CameraCalibrationCore/Source/CameraCalibrationCoreEditor/Private/LensComponentDetailCustomization.h b/Engine/Plugins/VirtualProduction/CameraCalibrationCore/Source/CameraCalibrationCoreEditor/Private/LensComponentDetailCustomization.h
new file mode 100644
--- /dev/null
+++ b/Engine/Plugins/VirtualProduction/CameraCalibrationCore/Source/CameraCalibrationCoreEditor/Private/LensComponentDetailCustomization.h
@@ -0,0 +1,40 @@
+// Copyright Epic Games, Inc. All Rights Reserved.
+
+#pragma once
+
+#include "IDetailCustomization.h"
+
+#include "LensComponent.h"
+
+/**
+* Customizes a ULensComponent details
+*/
+class FLensComponentDetailCustomization : public IDetailCustomization
+{
+public:
+
+	static TSharedRef<IDetailCustomization> MakeInstance()
+	{
+		return MakeShared<FLensComponentDetailCustomization>();
+	}
+
+	~FLensComponentDetailCustomization();
+
+	//~ Begin IDetailCustomization interface
+	virtual void CustomizeDetails(class IDetailLayoutBuilder& DetailBuilder) override;
+	//~ End IDetailCustomization interface
+
+protected:
+	/** Forces the details panel to refresh when the lens model changes */
+	void OnLensModelChanged();
+
+	/** Forces the details panel to refresh when the lens model changes */
+	void OnLensModelChanged(const TSubclassOf<ULensModel>& LensModel);
+
+protected:
+	/** Keep a reference to force refresh the layout */
+	IDetailLayoutBuilder* DetailLayout = nullptr;
+
+	/** Weak pointer to the lens component being edited */
+	TWeakObjectPtr<ULensComponent> WeakLensComponent;
+};
diff --git a/Engine/Plugins/VirtualProduction/CameraCalibrationCore/Source/CameraCalibrationCoreEditor/Private/LensDistortionStateDetailCustomization.cpp b/Engine/Plugins/VirtualProduction/CameraCalibrationCore/Source/CameraCalibrationCoreEditor/Private/LensDistortionStateDetailCustomization.cpp
new file mode 100644
--- /dev/null
+++ b/Engine/Plugins/VirtualProduction/CameraCalibrationCore/Source/CameraCalibrationCoreEditor/Private/LensDistortionStateDetailCustomization.cpp
@@ -0,0 +1,123 @@
+// Copyright Epic Games, Inc. All Rights Reserved.
+
+#include "LensDistortionStateDetailCustomization.h"
+
+#include "DetailWidgetRow.h"
+#include "IDetailChildrenBuilder.h"
+#include "IPropertyUtilities.h"
+#include "LensDistortionModelHandlerBase.h"
+
+FLensDistortionStateDetailCustomization::FLensDistortionStateDetailCustomization(TSubclassOf<ULensModel> InLensModel)
+	: LensModel(InLensModel)
+{
+}
+
+void FLensDistortionStateDetailCustomization::CustomizeHeader(TSharedRef<IPropertyHandle> InPropertyHandle, FDetailWidgetRow& HeaderRow, IPropertyTypeCustomizationUtils& CustomizationUtils)
+{
+	TSharedPtr<IPropertyUtilities> PropertyUtils = CustomizationUtils.GetPropertyUtilities();
+
+	HeaderRow.NameContent()
+	[
+		InPropertyHandle->CreatePropertyNameWidget()
+	]
+	.ValueContent()
+	[
+		InPropertyHandle->CreatePropertyValueWidget()
+	];
+}
+
+void FLensDistortionStateDetailCustomization::CustomizeChildren(TSharedRef<IPropertyHandle> PropertyHandle, IDetailChildrenBuilder& ChildBuilder, IPropertyTypeCustomizationUtils& CustomizationUtils)
+{
+	TSharedPtr<IPropertyUtilities> PropertyUtils = CustomizationUtils.GetPropertyUtilities();
+
+	uint32 NumChildren;
+	PropertyHandle->GetNumChildren(NumChildren);
+
+	for (uint32 Index = 0; Index < NumChildren; ++Index)
+	{
+		TSharedRef<IPropertyHandle> ChildPropertyHandle = PropertyHandle->GetChildHandle(Index).ToSharedRef();
+
+		const FName PropertyName = ChildPropertyHandle->GetProperty()->GetFName();
+		if (PropertyName == GET_MEMBER_NAME_CHECKED(FLensDistortionState, DistortionInfo))
+		{
+			CustomizeDistortionInfo(ChildPropertyHandle, ChildBuilder, CustomizationUtils);
+		}
+		else if (PropertyName == GET_MEMBER_NAME_CHECKED(FLensDistortionState, FocalLengthInfo))
+		{
+			CustomizeFocalLengthInfo(ChildPropertyHandle, ChildBuilder, CustomizationUtils);
+		}
+		else if (PropertyName == GET_MEMBER_NAME_CHECKED(FLensDistortionState, ImageCenter))
+		{
+			CustomizeImageCenterInfo(ChildPropertyHandle, ChildBuilder, CustomizationUtils);
+		}
+		else
+		{
+			ChildBuilder.AddProperty(ChildPropertyHandle).ShowPropertyButtons(false);
+		}
+	}
+}
+
+void FLensDistortionStateDetailCustomization::CustomizeDistortionInfo(TSharedRef<IPropertyHandle> PropertyHandle, IDetailChildrenBuilder& ChildBuilder, IPropertyTypeCustomizationUtils& CustomizationUtils)
+{
+	uint32 NumChildren;
+	PropertyHandle->GetNumChildren(NumChildren);
+
+	for (uint32 ChildIndex = 0; ChildIndex < NumChildren; ++ChildIndex)
+	{
+		TSharedRef<IPropertyHandle> ChildHandle = PropertyHandle->GetChildHandle(ChildIndex).ToSharedRef();
+		const FName PropertyName = ChildHandle->GetProperty()->GetFName();
+
+		if (PropertyName == GET_MEMBER_NAME_CHECKED(FDistortionInfo, Parameters))
+		{
+			TSharedPtr<IPropertyHandleArray> ParameterArrayHandle = ChildHandle->AsArray();
+			uint32 NumParameters = 0;
+			ParameterArrayHandle->GetNumElements(NumParameters);
+
+			if (LensModel)
+			{
+				const TArray<FText> ParameterNames = LensModel.GetDefaultObject()->GetParameterDisplayNames();
+
+				if (ParameterNames.Num() == NumParameters)
+				{
+					for (uint32 ParameterIndex = 0; ParameterIndex < NumParameters; ++ParameterIndex)
+					{
+						TSharedRef<IPropertyHandle> ParameterHandle = ParameterArrayHandle->GetElement(ParameterIndex);
+						ParameterHandle->SetPropertyDisplayName(ParameterNames[ParameterIndex]);
+					}
+				}
+			}
+				
+			ChildBuilder.AddProperty(ChildHandle).ShowPropertyButtons(false);
+		}
+		else
+		{
+			ChildBuilder.AddProperty(ChildHandle).ShowPropertyButtons(false);
+		}
+	}
+}
+
+void FLensDistortionStateDetailCustomization::CustomizeFocalLengthInfo(TSharedRef<IPropertyHandle> PropertyHandle, IDetailChildrenBuilder& ChildBuilder, IPropertyTypeCustomizationUtils& CustomizationUtils)
+{
+	// Effectively, this does a ShowInnerPropertiesOnly, which does not work great by default with nested structs like this
+	uint32 NumChildren;
+	PropertyHandle->GetNumChildren(NumChildren);
+
+	for (uint32 ChildIndex = 0; ChildIndex < NumChildren; ++ChildIndex)
+	{
+		TSharedRef<IPropertyHandle> ChildHandle = PropertyHandle->GetChildHandle(ChildIndex).ToSharedRef();
+		ChildBuilder.AddProperty(ChildHandle).ShowPropertyButtons(false);
+	}
+}
+
+void FLensDistortionStateDetailCustomization::CustomizeImageCenterInfo(TSharedRef<IPropertyHandle> PropertyHandle, IDetailChildrenBuilder& ChildBuilder, IPropertyTypeCustomizationUtils& CustomizationUtils)
+{
+	// Effectively, this does a ShowInnerPropertiesOnly, which does not work great by default with nested structs like this
+	uint32 NumChildren;
+	PropertyHandle->GetNumChildren(NumChildren);
+
+	for (uint32 ChildIndex = 0; ChildIndex < NumChildren; ++ChildIndex)
+	{
+		TSharedRef<IPropertyHandle> ChildHandle = PropertyHandle->GetChildHandle(ChildIndex).ToSharedRef();
+		ChildBuilder.AddProperty(ChildHandle).ShowPropertyButtons(false);
+	}
+}
diff --git a/Engine/Plugins/VirtualProduction/CameraCalibrationCore/Source/CameraCalibrationCoreEditor/Private/LensDistortionStateDetailCustomization.h b/Engine/Plugins/VirtualProduction/CameraCalibrationCore/Source/CameraCalibrationCoreEditor/Private/LensDistortionStateDetailCustomization.h
new file mode 100644
--- /dev/null
+++ b/Engine/Plugins/VirtualProduction/CameraCalibrationCore/Source/CameraCalibrationCoreEditor/Private/LensDistortionStateDetailCustomization.h
@@ -0,0 +1,32 @@
+// Copyright Epic Games, Inc. All Rights Reserved.
+
+#pragma once
+
+#include "IPropertyTypeCustomization.h"
+
+#include "Models/LensModel.h"
+#include "Templates/SubclassOf.h"
+
+class FLensDistortionStateDetailCustomization : public IPropertyTypeCustomization
+{
+public:
+	FLensDistortionStateDetailCustomization(TSubclassOf<ULensModel> InLensModel);
+
+	// Begin IDetailCustomization interface
+	virtual void CustomizeHeader(TSharedRef<IPropertyHandle> PropertyHandle, FDetailWidgetRow& HeaderRow, IPropertyTypeCustomizationUtils& CustomizationUtils) override;
+	virtual void CustomizeChildren(TSharedRef<IPropertyHandle> PropertyHandle, IDetailChildrenBuilder& ChildBuilder, IPropertyTypeCustomizationUtils& CustomizationUtils) override;
+	//~ End IDetailCustomization interface
+
+private:
+	/** Customize the FDistortionInfo struct property */
+	void CustomizeDistortionInfo(TSharedRef<IPropertyHandle> PropertyHandle, IDetailChildrenBuilder& ChildBuilder, IPropertyTypeCustomizationUtils& CustomizationUtils);
+
+	/** Customize the FFocalLengthInfo struct property */
+	void CustomizeFocalLengthInfo(TSharedRef<IPropertyHandle> PropertyHandle, IDetailChildrenBuilder& ChildBuilder, IPropertyTypeCustomizationUtils& CustomizationUtils);
+
+	/** Customize the FImageCenterInfo struct property */
+	void CustomizeImageCenterInfo(TSharedRef<IPropertyHandle> PropertyHandle, IDetailChildrenBuilder& ChildBuilder, IPropertyTypeCustomizationUtils& CustomizationUtils);
+
+	/** The LensModel used to customize the names of the properties in the distortion parameter array */
+	TSubclassOf<ULensModel> LensModel;
+};
diff --git a/Engine/Plugins/VirtualProduction/LiveLinkLens/Source/LiveLinkLens/Private/LiveLinkLensController.cpp b/Engine/Plugins/VirtualProduction/LiveLinkLens/Source/LiveLinkLens/Private/LiveLinkLensController.cpp
--- a/Engine/Plugins/VirtualProduction/LiveLinkLens/Source/LiveLinkLens/Private/LiveLinkLensController.cpp
+++ b/Engine/Plugins/VirtualProduction/LiveLinkLens/Source/LiveLinkLens/Private/LiveLinkLensController.cpp
@@ -38,18 +38,25 @@ void ULiveLinkLensController::Tick(float DeltaTime, const FLiveLinkSubjectFrameD
 	{
 		if (ULensComponent* const LensComponent = Cast<ULensComponent>(GetAttachedComponent()))
 		{
-			UCameraCalibrationSubsystem* const SubSystem = GEngine->GetEngineSubsystem<UCameraCalibrationSubsystem>();
-			const TSubclassOf<ULensModel> LensModel = SubSystem->GetRegisteredLensModel(StaticData->LensModel);
-			LensComponent->SetLensModel(LensModel);
+			if (LensComponent->GetDistortionSource() == EDistortionSource::LiveLinkLensSubject)
+			{
+				UCameraCalibrationSubsystem* const SubSystem = GEngine->GetEngineSubsystem<UCameraCalibrationSubsystem>();
+				const TSubclassOf<ULensModel> LensModel = SubSystem->GetRegisteredLensModel(StaticData->LensModel);
 
-			// Update the lens distortion handler with the latest frame of data from the LiveLink source
-			FLensDistortionState DistortionState;
+				if (LensComponent->GetLensModel() != LensModel)
+				{
+					LensComponent->SetLensModel(LensModel);
+				}
 
-			DistortionState.DistortionInfo.Parameters = FrameData->DistortionParameters;
-			DistortionState.FocalLengthInfo.FxFy = FrameData->FxFy;
-			DistortionState.ImageCenter.PrincipalPoint = FrameData->PrincipalPoint;
+				// Update the lens distortion handler with the latest frame of data from the LiveLink source
+				FLensDistortionState DistortionState;
 
-			LensComponent->SetDistortionState(DistortionState);
+				DistortionState.DistortionInfo.Parameters = FrameData->DistortionParameters;
+				DistortionState.FocalLengthInfo.FxFy = FrameData->FxFy;
+				DistortionState.ImageCenter.PrincipalPoint = FrameData->PrincipalPoint;
+
+				LensComponent->SetDistortionState(DistortionState);
+			}
 		}
 	}
 }
