CameraCalibration: Anamorphic model spams warnings about missing undistortion displacement map material
#rnx

diff --git a/Engine/Plugins/VirtualProduction/CameraCalibrationCore/Source/CameraCalibrationCore/Private/AnamorphicLensDistortionModelHandler.cpp b/Engine/Plugins/VirtualProduction/CameraCalibrationCore/Source/CameraCalibrationCore/Private/AnamorphicLensDistortionModelHandler.cpp
--- a/Engine/Plugins/VirtualProduction/CameraCalibrationCore/Source/CameraCalibrationCore/Private/AnamorphicLensDistortionModelHandler.cpp
+++ b/Engine/Plugins/VirtualProduction/CameraCalibrationCore/Source/CameraCalibrationCore/Private/AnamorphicLensDistortionModelHandler.cpp
@@ -70,7 +70,11 @@ void UAnamorphicLensDistortionModelHandler::InitDistortionMaterials()
 		DistortionPostProcessMID = UMaterialInstanceDynamic::Create(DistortionMaterialParent, this);
 	}
 
-	// TODO: Add an undistortion displacement map MID once the material is authored
+	if (UndistortionDisplacementMapMID == nullptr)
+	{
+		UMaterialInterface* MaterialParent = GetDefault<UCameraCalibrationSettings>()->GetDefaultUndistortionDisplacementMaterial(this->StaticClass());
+		UndistortionDisplacementMapMID = UMaterialInstanceDynamic::Create(MaterialParent, this);
+	}
 
 	if (DistortionDisplacementMapMID == nullptr)
 	{
@@ -78,7 +82,7 @@ void UAnamorphicLensDistortionModelHandler::InitDistortionMaterials()
 		DistortionDisplacementMapMID = UMaterialInstanceDynamic::Create(MaterialParent, this);
 	}
 
-	// TODO: Add the UndistortionDisplacementMap parameter
+	DistortionPostProcessMID->SetTextureParameterValue("UndistortionDisplacementMap", UndistortionDisplacementMapRT);
 	DistortionPostProcessMID->SetTextureParameterValue("DistortionDisplacementMap", DistortionDisplacementMapRT);
 
 	SetDistortionState(CurrentState);
@@ -112,7 +116,7 @@ void UAnamorphicLensDistortionModelHandler::UpdateMaterialParameters()
 		MID->SetScalarParameterValue("phi_mnt", AnamorphicParameters.LensRotation);
 	};
 
-	// TODO: Update the undistortion material parameters
+	SetDistortionMaterialParameters(UndistortionDisplacementMapMID);
 	SetDistortionMaterialParameters(DistortionDisplacementMapMID);
 }
 
diff --git a/Engine/Plugins/VirtualProduction/CameraCalibrationCore/Source/CameraCalibrationCore/Private/CameraCalibrationSettings.cpp b/Engine/Plugins/VirtualProduction/CameraCalibrationCore/Source/CameraCalibrationCore/Private/CameraCalibrationSettings.cpp
--- a/Engine/Plugins/VirtualProduction/CameraCalibrationCore/Source/CameraCalibrationCore/Private/CameraCalibrationSettings.cpp
+++ b/Engine/Plugins/VirtualProduction/CameraCalibrationCore/Source/CameraCalibrationCore/Private/CameraCalibrationSettings.cpp
@@ -16,7 +16,7 @@ UCameraCalibrationSettings::UCameraCalibrationSettings()
 	DefaultDistortionMaterials.Add(USphericalLensDistortionModelHandler::StaticClass(), TSoftObjectPtr<UMaterialInterface>(FSoftObjectPath(TEXT("/CameraCalibrationCore/Materials/M_DistortionPostProcess.M_DistortionPostProcess"))));
 
 	// Default materials to use for the anamorphic lens model
-	// TODO: Add default undistortion displacement map material once it is authored
+	DefaultUndistortionDisplacementMaterials.Add(UAnamorphicLensDistortionModelHandler::StaticClass(), TSoftObjectPtr<UMaterialInterface>(FSoftObjectPath(TEXT("/CameraCalibrationCore/Materials/M_AnamorphicUndistortionDisplacementMap.M_AnamorphicUndistortionDisplacementMap"))));
 	DefaultDistortionDisplacementMaterials.Add(UAnamorphicLensDistortionModelHandler::StaticClass(), TSoftObjectPtr<UMaterialInterface>(FSoftObjectPath(TEXT("/CameraCalibrationCore/Materials/M_AnamorphicDistortionDisplacementMap.M_AnamorphicDistortionDisplacementMap"))));
 	DefaultDistortionMaterials.Add(UAnamorphicLensDistortionModelHandler::StaticClass(), TSoftObjectPtr<UMaterialInterface>(FSoftObjectPath(TEXT("/CameraCalibrationCore/Materials/M_DistortionPostProcess.M_DistortionPostProcess"))));
 
