CameraCalibration - Delete deprecated code from 5.1 and earlier
#rnx

diff --git a/Engine/Plugins/VirtualProduction/CameraCalibrationCore/Source/CameraCalibrationCore/Private/CameraCalibrationSubsystem.cpp b/Engine/Plugins/VirtualProduction/CameraCalibrationCore/Source/CameraCalibrationCore/Private/CameraCalibrationSubsystem.cpp
--- a/Engine/Plugins/VirtualProduction/CameraCalibrationCore/Source/CameraCalibrationCore/Private/CameraCalibrationSubsystem.cpp
+++ b/Engine/Plugins/VirtualProduction/CameraCalibrationCore/Source/CameraCalibrationCore/Private/CameraCalibrationSubsystem.cpp
@@ -129,22 +129,6 @@ void UCameraCalibrationSubsystem::UnregisterOverlayMaterial(const FName& Materia
 	RegisteredOverlayMaterials.Remove(MaterialName);
 }
 
-void UCameraCalibrationSubsystem::UpdateOriginalFocalLength(UCineCameraComponent* Component, float InFocalLength)
-{
-	// This function has been deprecated.
-}
-
-void UCameraCalibrationSubsystem::UpdateOverscanFocalLength(UCineCameraComponent* Component, float InFocalLength)
-{
-	// This function has been deprecated.
-}
-
-bool UCameraCalibrationSubsystem::GetOriginalFocalLength(UCineCameraComponent* Component, float& OutFocalLength)
-{
-	// This function has been deprecated.
-	return false;
-}
-
 TSubclassOf<ULensModel> UCameraCalibrationSubsystem::GetRegisteredLensModel(FName ModelName) const
 {
 	if (LensModelMap.Contains(ModelName))
diff --git a/Engine/Plugins/VirtualProduction/CameraCalibrationCore/Source/CameraCalibrationCore/Private/LensComponent.cpp b/Engine/Plugins/VirtualProduction/CameraCalibrationCore/Source/CameraCalibrationCore/Private/LensComponent.cpp
--- a/Engine/Plugins/VirtualProduction/CameraCalibrationCore/Source/CameraCalibrationCore/Private/LensComponent.cpp
+++ b/Engine/Plugins/VirtualProduction/CameraCalibrationCore/Source/CameraCalibrationCore/Private/LensComponent.cpp
@@ -330,11 +330,6 @@ void ULensComponent::PostEditChangeProperty(struct FPropertyChangedEvent& Proper
 }
 #endif //WITH_EDITOR
 
-void ULensComponent::OnPostActorTick(UWorld* World, ELevelTick TickType, float DeltaSeconds)
-{
-	// The use of this callback by this class has been deprecated
-}
-
 void ULensComponent::ReapplyNodalOffset()
 {
 	// Find the TrackedComponent based on the serialized name
diff --git a/Engine/Plugins/VirtualProduction/CameraCalibrationCore/Source/CameraCalibrationCore/Private/LensFile.cpp b/Engine/Plugins/VirtualProduction/CameraCalibrationCore/Source/CameraCalibrationCore/Private/LensFile.cpp
--- a/Engine/Plugins/VirtualProduction/CameraCalibrationCore/Source/CameraCalibrationCore/Private/LensFile.cpp
+++ b/Engine/Plugins/VirtualProduction/CameraCalibrationCore/Source/CameraCalibrationCore/Private/LensFile.cpp
@@ -1005,11 +1005,6 @@ void ULensFile::OnDistortionDerivedDataJobCompleted(const FDerivedDistortionData
 	}
 }
 
-bool ULensFile::IsCineCameraCompatible(const UCineCameraComponent* CineCameraComponent) const
-{
-	return true;
-}
-
 void ULensFile::UpdateInputTolerance(const float NewTolerance)
 {
 	InputTolerance = NewTolerance;
diff --git a/Engine/Plugins/VirtualProduction/CameraCalibrationCore/Source/CameraCalibrationCore/Public/CameraCalibrationSubsystem.h b/Engine/Plugins/VirtualProduction/CameraCalibrationCore/Source/CameraCalibrationCore/Public/CameraCalibrationSubsystem.h
--- a/Engine/Plugins/VirtualProduction/CameraCalibrationCore/Source/CameraCalibrationCore/Public/CameraCalibrationSubsystem.h
+++ b/Engine/Plugins/VirtualProduction/CameraCalibrationCore/Source/CameraCalibrationCore/Public/CameraCalibrationSubsystem.h
@@ -19,13 +19,6 @@
 
 class UCineCameraComponent;
 
-struct UE_DEPRECATED(5.1, "This struct has been deprecated.") FCachedFocalLength
-{
-public:
-	float OriginalFocalLength = 0.0f;
-	float OverscanFocalLength = 0.0f;
-};
-
 /**
  * Camera Calibration subsystem
  */
@@ -118,29 +111,12 @@ public:
 	/** Remove a Lens Model from the registered model map */
 	void UnregisterDistortionModel(TSubclassOf<ULensModel> LensModel);
 
-	UE_DEPRECATED(5.1, "This function has been deprecated. The subsystem no longer tracks original focal length for distortion.")
-	void UpdateOriginalFocalLength(UCineCameraComponent* Component, float InFocalLength);
-
-	UE_DEPRECATED(5.1, "This function has been deprecated. The subsystem no longer tracks overscanned focal length for distortion.")
-	void UpdateOverscanFocalLength(UCineCameraComponent* Component, float InFocalLength);
-
 	/** Register a new overlay material name and path that can be queried from camera calibration tools */
 	void RegisterOverlayMaterial(const FName& MaterialName, const FSoftObjectPath& MaterialPath);
 
-	UE_DEPRECATED(5.1, "FNames containing full object paths are deprecated. Use FSoftObjectPath instead.")
-	void RegisterOverlayMaterial(const FName& MaterialName, const FName& MaterialPathName)
-	{
-PRAGMA_DISABLE_DEPRECATION_WARNINGS
-		RegisterOverlayMaterial(MaterialName, FSoftObjectPath(MaterialPathName));
-PRAGMA_ENABLE_DEPRECATION_WARNINGS
-	}
-
 	/** Unregister an overlay material */
 	void UnregisterOverlayMaterial(const FName& MaterialName);
 
-	UE_DEPRECATED(5.1, "This function has been deprecated. The subsystem no longer tracks original focal length for distortion.")
-	bool GetOriginalFocalLength(UCineCameraComponent* Component, float& OutFocalLength);
-
 private:
 	/** Default lens file to use when no override has been provided */
 	UPROPERTY(Transient)
diff --git a/Engine/Plugins/VirtualProduction/CameraCalibrationCore/Source/CameraCalibrationCore/Public/LensComponent.h b/Engine/Plugins/VirtualProduction/CameraCalibrationCore/Source/CameraCalibrationCore/Public/LensComponent.h
--- a/Engine/Plugins/VirtualProduction/CameraCalibrationCore/Source/CameraCalibrationCore/Public/LensComponent.h
+++ b/Engine/Plugins/VirtualProduction/CameraCalibrationCore/Source/CameraCalibrationCore/Public/LensComponent.h
@@ -217,10 +217,6 @@ PRAGMA_DISABLE_DEPRECATION_WARNINGS
 	FDistortionHandlerPicker GetDistortionHandlerPicker() const;
 PRAGMA_ENABLE_DEPRECATION_WARNINGS
 
-protected:
-	UE_DEPRECATED(5.1, "The use of this callback by this class has been deprecated and it is no longer registered. You can register your own delegate with FWorldDelegates::OnWorldPostActorTick")
-	void OnPostActorTick(UWorld* World, ELevelTick TickType, float DeltaSeconds);
-
 private:
 	/** Evaluate the LensFile for nodal offset (using the current evaluation mode) and apply it to the latest component to offset */
 	void ApplyNodalOffset();
@@ -358,18 +354,6 @@ protected:
 	UPROPERTY()
 	FVector OriginalCameraLocation_DEPRECATED;
 
-	UE_DEPRECATED(5.1, "This property has been deprecated. Use the DistortionStateSource to specify whether the lens file should be evaluated for distortion.")
-	UPROPERTY()
-	bool bEvaluateLensFileForDistortion_DEPRECATED = false;
-
-	UE_DEPRECATED(5.1, "This property has been deprecated. Use the LensDistortionHandlerMap to get a handler for the current LensModel.")
-	UPROPERTY(Transient)
-	TObjectPtr<ULensDistortionModelHandlerBase> LensDistortionHandler_DEPRECATED = nullptr;
-
-	UE_DEPRECATED(5.1, "This property has been deprecated. Producer GUIDs are no longer used to identify distortion handlers. Use the LensDistortionHandlerMap to get a handler for the current LensModel.")
-	UPROPERTY(DuplicateTransient)
-	FGuid DistortionProducerID_DEPRECATED;
-
 PRAGMA_DISABLE_DEPRECATION_WARNINGS
 	UE_DEPRECATED(5.1, "This property has been deprecated. The handler picker is no longer used to identify a distortion handler. Use the LensDistortionHandlerMap to get a handler for the current LensModel.")
 	UPROPERTY(meta=(DeprecatedProperty, DeprecationMessage="This property has been deprecated. Use GetDistortionSource() and SetDistortionSource() instead."))
diff --git a/Engine/Plugins/VirtualProduction/CameraCalibrationCore/Source/CameraCalibrationCore/Public/LensFile.h b/Engine/Plugins/VirtualProduction/CameraCalibrationCore/Source/CameraCalibrationCore/Public/LensFile.h
--- a/Engine/Plugins/VirtualProduction/CameraCalibrationCore/Source/CameraCalibrationCore/Public/LensFile.h
+++ b/Engine/Plugins/VirtualProduction/CameraCalibrationCore/Source/CameraCalibrationCore/Public/LensFile.h
@@ -123,9 +123,6 @@ public:
 	/** Update the resolution used for intermediate blending displacement maps and for STMap derived data */
 	void UpdateDisplacementMapResolution(const FIntPoint NewDisplacementMapResolution);
 
-	UE_DEPRECATED(5.1, "This function is deprecated. The LensFile checks this internally when it evaluates for distortion.")
-	bool IsCineCameraCompatible(const UCineCameraComponent* CineCameraComponent) const;
-
 	/** Update the input tolerance used when adding points to calibration tables */
 	void UpdateInputTolerance(const float NewTolerance);
 
