CameraCalibration - Fix Lens data tables to use input tolerance when querying for points
#rnx

diff --git a/Engine/Plugins/VirtualProduction/CameraCalibrationCore/Source/CameraCalibrationCore/Private/LensFile.cpp b/Engine/Plugins/VirtualProduction/CameraCalibrationCore/Source/CameraCalibrationCore/Private/LensFile.cpp
--- a/Engine/Plugins/VirtualProduction/CameraCalibrationCore/Source/CameraCalibrationCore/Private/LensFile.cpp
+++ b/Engine/Plugins/VirtualProduction/CameraCalibrationCore/Source/CameraCalibrationCore/Private/LensFile.cpp
@@ -1037,27 +1037,27 @@ TArray<FNodalOffsetPointInfo> ULensFile::GetNodalOffsetPoints() const
 
 bool ULensFile::GetDistortionPoint(const float InFocus, const float InZoom, FDistortionInfo& OutDistortionInfo) const
 {
-	return DistortionTable.GetPoint(InFocus, InZoom, OutDistortionInfo);
+	return DistortionTable.GetPoint(InFocus, InZoom, OutDistortionInfo, InputTolerance);
 }
 
 bool ULensFile::GetFocalLengthPoint(const float InFocus, const float InZoom, FFocalLengthInfo& OutFocalLengthInfo) const
 {
-	return FocalLengthTable.GetPoint(InFocus, InZoom, OutFocalLengthInfo);
+	return FocalLengthTable.GetPoint(InFocus, InZoom, OutFocalLengthInfo, InputTolerance);
 }
 
 bool ULensFile::GetImageCenterPoint(const float InFocus, const float InZoom, FImageCenterInfo& OutImageCenterInfo) const
 {
-	return ImageCenterTable.GetPoint(InFocus, InZoom, OutImageCenterInfo);
+	return ImageCenterTable.GetPoint(InFocus, InZoom, OutImageCenterInfo, InputTolerance);
 }
 
 bool ULensFile::GetNodalOffsetPoint(const float InFocus, const float InZoom, FNodalPointOffset& OutNodalPointOffset) const
 {
-	return NodalOffsetTable.GetPoint(InFocus, InZoom, OutNodalPointOffset);
+	return NodalOffsetTable.GetPoint(InFocus, InZoom, OutNodalPointOffset, InputTolerance);
 }
 
 bool ULensFile::GetSTMapPoint(const float InFocus, const float InZoom, FSTMapInfo& OutSTMapInfo) const
 {
-	return STMapTable.GetPoint(InFocus, InZoom, OutSTMapInfo);
+	return STMapTable.GetPoint(InFocus, InZoom, OutSTMapInfo, InputTolerance);
 }
 
 void ULensFile::AddDistortionPoint(float NewFocus, float NewZoom, const FDistortionInfo& NewDistortionPoint, const FFocalLengthInfo& NewFocalLength)
diff --git a/Engine/Plugins/VirtualProduction/CameraCalibrationCore/Source/CameraCalibrationCore/Private/Tables/BaseLensTable.cpp b/Engine/Plugins/VirtualProduction/CameraCalibrationCore/Source/CameraCalibrationCore/Private/Tables/BaseLensTable.cpp
--- a/Engine/Plugins/VirtualProduction/CameraCalibrationCore/Source/CameraCalibrationCore/Private/Tables/BaseLensTable.cpp
+++ b/Engine/Plugins/VirtualProduction/CameraCalibrationCore/Source/CameraCalibrationCore/Private/Tables/BaseLensTable.cpp
@@ -74,7 +74,7 @@ bool FBaseLensTable::HasLinkedFocusValues(const float InFocus, float InputTolera
 			return false;
 		}
 
-		if (LinkDataTable->DoesFocusPointExists(InFocus))
+		if (LinkDataTable->DoesFocusPointExists(InFocus, InputTolerance))
 		{
 			return true;
 		}
diff --git a/Engine/Plugins/VirtualProduction/CameraCalibrationCore/Source/CameraCalibrationCore/Private/Tables/DistortionParametersTable.cpp b/Engine/Plugins/VirtualProduction/CameraCalibrationCore/Source/CameraCalibrationCore/Private/Tables/DistortionParametersTable.cpp
--- a/Engine/Plugins/VirtualProduction/CameraCalibrationCore/Source/CameraCalibrationCore/Private/Tables/DistortionParametersTable.cpp
+++ b/Engine/Plugins/VirtualProduction/CameraCalibrationCore/Source/CameraCalibrationCore/Private/Tables/DistortionParametersTable.cpp
@@ -138,14 +138,14 @@ bool FDistortionTable::BuildParameterCurve(float InFocus, int32 ParameterIndex,
 	return false;
 }
 
-const FDistortionFocusPoint* FDistortionTable::GetFocusPoint(float InFocus) const
+const FDistortionFocusPoint* FDistortionTable::GetFocusPoint(float InFocus, float InputTolerance) const
 {
-	return FocusPoints.FindByPredicate([InFocus](const FDistortionFocusPoint& Point) { return FMath::IsNearlyEqual(Point.Focus, InFocus); });
+	return FocusPoints.FindByPredicate([InFocus, InputTolerance](const FDistortionFocusPoint& Point) { return FMath::IsNearlyEqual(Point.Focus, InFocus, InputTolerance); });
 }
 
-FDistortionFocusPoint* FDistortionTable::GetFocusPoint(float InFocus)
+FDistortionFocusPoint* FDistortionTable::GetFocusPoint(float InFocus, float InputTolerance)
 {
-	return FocusPoints.FindByPredicate([InFocus](const FDistortionFocusPoint& Point) { return FMath::IsNearlyEqual(Point.Focus, InFocus); });
+	return FocusPoints.FindByPredicate([InFocus, InputTolerance](const FDistortionFocusPoint& Point) { return FMath::IsNearlyEqual(Point.Focus, InFocus, InputTolerance); });
 }
 
 void FDistortionTable::ForEachPoint(FFocusPointCallback InCallback) const
@@ -176,9 +176,9 @@ void FDistortionTable::RemoveZoomPoint(float InFocus, float InZoom)
 	LensDataTableUtils::RemoveZoomPoint(FocusPoints, InFocus, InZoom);
 }
 
-bool FDistortionTable::DoesFocusPointExists(float InFocus) const
+bool FDistortionTable::DoesFocusPointExists(float InFocus, float InputTolerance) const
 {
-	if (GetFocusPoint(InFocus) != nullptr)
+	if (GetFocusPoint(InFocus, InputTolerance) != nullptr)
 	{
 		return true;
 	}
@@ -225,7 +225,7 @@ bool FDistortionTable::AddPoint(float InFocus, float InZoom, const FDistortionIn
 
 bool FDistortionTable::GetPoint(const float InFocus, const float InZoom, FDistortionInfo& OutData, float InputTolerance) const
 {
-	if (const FDistortionFocusPoint* DistortionFocusPoint = GetFocusPoint(InFocus))
+	if (const FDistortionFocusPoint* DistortionFocusPoint = GetFocusPoint(InFocus, InputTolerance))
 	{
 		FDistortionInfo DistortionInfo;
 		if (DistortionFocusPoint->GetPoint(InZoom, DistortionInfo, InputTolerance))
diff --git a/Engine/Plugins/VirtualProduction/CameraCalibrationCore/Source/CameraCalibrationCore/Private/Tables/FocalLengthTable.cpp b/Engine/Plugins/VirtualProduction/CameraCalibrationCore/Source/CameraCalibrationCore/Private/Tables/FocalLengthTable.cpp
--- a/Engine/Plugins/VirtualProduction/CameraCalibrationCore/Source/CameraCalibrationCore/Private/Tables/FocalLengthTable.cpp
+++ b/Engine/Plugins/VirtualProduction/CameraCalibrationCore/Source/CameraCalibrationCore/Private/Tables/FocalLengthTable.cpp
@@ -183,14 +183,14 @@ bool FFocalLengthTable::BuildParameterCurve(float InFocus, int32 ParameterIndex,
 	return false;
 }
 
-const FFocalLengthFocusPoint* FFocalLengthTable::GetFocusPoint(float InFocus) const
+const FFocalLengthFocusPoint* FFocalLengthTable::GetFocusPoint(float InFocus, float InputTolerance) const
 {
-	return FocusPoints.FindByPredicate([InFocus](const FFocalLengthFocusPoint& Point) { return FMath::IsNearlyEqual(Point.Focus, InFocus); });
+	return FocusPoints.FindByPredicate([InFocus, InputTolerance](const FFocalLengthFocusPoint& Point) { return FMath::IsNearlyEqual(Point.Focus, InFocus, InputTolerance); });
 }
 
-FFocalLengthFocusPoint* FFocalLengthTable::GetFocusPoint(float InFocus)
+FFocalLengthFocusPoint* FFocalLengthTable::GetFocusPoint(float InFocus, float InputTolerance)
 {
-	return FocusPoints.FindByPredicate([InFocus](const FFocalLengthFocusPoint& Point) { return FMath::IsNearlyEqual(Point.Focus, InFocus); });
+	return FocusPoints.FindByPredicate([InFocus, InputTolerance](const FFocalLengthFocusPoint& Point) { return FMath::IsNearlyEqual(Point.Focus, InFocus, InputTolerance); });
 }
 
 TConstArrayView<FFocalLengthFocusPoint> FFocalLengthTable::GetFocusPoints() const
@@ -216,9 +216,9 @@ void FFocalLengthTable::RemoveZoomPoint(float InFocus, float InZoom)
 	LensDataTableUtils::RemoveZoomPoint(FocusPoints, InFocus, InZoom);
 }
 
-bool FFocalLengthTable::DoesFocusPointExists(float InFocus) const
+bool FFocalLengthTable::DoesFocusPointExists(float InFocus, float InputTolerance) const
 {
-	if (GetFocusPoint(InFocus) != nullptr)
+	if (GetFocusPoint(InFocus, InputTolerance) != nullptr)
 	{
 		return true;
 	}
@@ -233,7 +233,7 @@ bool FFocalLengthTable::AddPoint(float InFocus, float InZoom, const FFocalLength
 
 bool FFocalLengthTable::GetPoint(const float InFocus, const float InZoom, FFocalLengthInfo& OutData, float InputTolerance) const
 {
-	if (const FFocalLengthFocusPoint* FocalLengthFocusPoint = GetFocusPoint(InFocus))
+	if (const FFocalLengthFocusPoint* FocalLengthFocusPoint = GetFocusPoint(InFocus, InputTolerance))
 	{
 		FFocalLengthZoomPoint ZoomPont;
 		if (FocalLengthFocusPoint->GetPoint(InZoom, ZoomPont, InputTolerance))
diff --git a/Engine/Plugins/VirtualProduction/CameraCalibrationCore/Source/CameraCalibrationCore/Private/Tables/ImageCenterTable.cpp b/Engine/Plugins/VirtualProduction/CameraCalibrationCore/Source/CameraCalibrationCore/Private/Tables/ImageCenterTable.cpp
--- a/Engine/Plugins/VirtualProduction/CameraCalibrationCore/Source/CameraCalibrationCore/Private/Tables/ImageCenterTable.cpp
+++ b/Engine/Plugins/VirtualProduction/CameraCalibrationCore/Source/CameraCalibrationCore/Private/Tables/ImageCenterTable.cpp
@@ -98,9 +98,9 @@ void FImageCenterTable::ForEachPoint(FFocusPointCallback InCallback) const
 	}
 }
 
-bool FImageCenterTable::DoesFocusPointExists(float InFocus) const
+bool FImageCenterTable::DoesFocusPointExists(float InFocus, float InputTolerance) const
 {
-	if (GetFocusPoint(InFocus) != nullptr)
+	if (GetFocusPoint(InFocus, InputTolerance) != nullptr)
 	{
 		return true;
 	}
@@ -139,14 +139,14 @@ bool FImageCenterTable::BuildParameterCurve(float InFocus, int32 ParameterIndex,
 	return false;
 }
 
-const FImageCenterFocusPoint* FImageCenterTable::GetFocusPoint(float InFocus) const
+const FImageCenterFocusPoint* FImageCenterTable::GetFocusPoint(float InFocus, float InputTolerance) const
 {
-	return FocusPoints.FindByPredicate([InFocus](const FImageCenterFocusPoint& Points) { return FMath::IsNearlyEqual(Points.Focus, InFocus); });
+	return FocusPoints.FindByPredicate([InFocus, InputTolerance](const FImageCenterFocusPoint& Points) { return FMath::IsNearlyEqual(Points.Focus, InFocus, InputTolerance); });
 }
 
-FImageCenterFocusPoint* FImageCenterTable::GetFocusPoint(float InFocus)
+FImageCenterFocusPoint* FImageCenterTable::GetFocusPoint(float InFocus, float InputTolerance)
 {
-	return FocusPoints.FindByPredicate([InFocus](const FImageCenterFocusPoint& Points) { return FMath::IsNearlyEqual(Points.Focus, InFocus); });
+	return FocusPoints.FindByPredicate([InFocus, InputTolerance](const FImageCenterFocusPoint& Points) { return FMath::IsNearlyEqual(Points.Focus, InFocus, InputTolerance); });
 }
 
 TConstArrayView<FImageCenterFocusPoint> FImageCenterTable::GetFocusPoints() const
@@ -210,7 +210,7 @@ bool FImageCenterTable::AddPoint(float InFocus, float InZoom, const FImageCenter
 
 bool FImageCenterTable::GetPoint(const float InFocus, const float InZoom, FImageCenterInfo& OutData, float InputTolerance) const
 {
-	if (const FImageCenterFocusPoint* FocalLengthFocusPoint = GetFocusPoint(InFocus))
+	if (const FImageCenterFocusPoint* FocalLengthFocusPoint = GetFocusPoint(InFocus, InputTolerance))
 	{
 		FImageCenterInfo ImageCenterInfo;
 		if (FocalLengthFocusPoint->GetPoint(InZoom, ImageCenterInfo, InputTolerance))
diff --git a/Engine/Plugins/VirtualProduction/CameraCalibrationCore/Source/CameraCalibrationCore/Private/Tables/NodalOffsetTable.cpp b/Engine/Plugins/VirtualProduction/CameraCalibrationCore/Source/CameraCalibrationCore/Private/Tables/NodalOffsetTable.cpp
--- a/Engine/Plugins/VirtualProduction/CameraCalibrationCore/Source/CameraCalibrationCore/Private/Tables/NodalOffsetTable.cpp
+++ b/Engine/Plugins/VirtualProduction/CameraCalibrationCore/Source/CameraCalibrationCore/Private/Tables/NodalOffsetTable.cpp
@@ -206,14 +206,14 @@ bool FNodalOffsetTable::BuildParameterCurve(float InFocus, int32 ParameterIndex,
 	return false;
 }
 
-const FNodalOffsetFocusPoint* FNodalOffsetTable::GetFocusPoint(float InFocus) const
+const FNodalOffsetFocusPoint* FNodalOffsetTable::GetFocusPoint(float InFocus, float InputTolerance) const
 {
-	return FocusPoints.FindByPredicate([InFocus](const FNodalOffsetFocusPoint& Point) { return FMath::IsNearlyEqual(Point.Focus, InFocus); });
+	return FocusPoints.FindByPredicate([InFocus, InputTolerance](const FNodalOffsetFocusPoint& Point) { return FMath::IsNearlyEqual(Point.Focus, InFocus, InputTolerance); });
 }
 
-FNodalOffsetFocusPoint* FNodalOffsetTable::GetFocusPoint(float InFocus)
+FNodalOffsetFocusPoint* FNodalOffsetTable::GetFocusPoint(float InFocus, float InputTolerance)
 {
-	return FocusPoints.FindByPredicate([InFocus](const FNodalOffsetFocusPoint& Point) { return FMath::IsNearlyEqual(Point.Focus, InFocus); });
+	return FocusPoints.FindByPredicate([InFocus, InputTolerance](const FNodalOffsetFocusPoint& Point) { return FMath::IsNearlyEqual(Point.Focus, InFocus, InputTolerance); });
 }
 
 TConstArrayView<FNodalOffsetFocusPoint> FNodalOffsetTable::GetFocusPoints() const
@@ -244,9 +244,9 @@ void FNodalOffsetTable::RemoveZoomPoint(float InFocus, float InZoom)
 	LensDataTableUtils::RemoveZoomPoint(FocusPoints, InFocus, InZoom);
 }
 
-bool FNodalOffsetTable::DoesFocusPointExists(float InFocus) const
+bool FNodalOffsetTable::DoesFocusPointExists(float InFocus, float InputTolerance) const
 {
-	if (GetFocusPoint(InFocus) != nullptr)
+	if (GetFocusPoint(InFocus, InputTolerance) != nullptr)
 	{
 		return true;
 	}
@@ -261,7 +261,7 @@ bool FNodalOffsetTable::AddPoint(float InFocus, float InZoom, const FNodalPointO
 
 bool FNodalOffsetTable::GetPoint(const float InFocus, const float InZoom, FNodalPointOffset& OutData, float InputTolerance) const
 {
-	if (const FNodalOffsetFocusPoint* NodalOffsetFocusPoint = GetFocusPoint(InFocus))
+	if (const FNodalOffsetFocusPoint* NodalOffsetFocusPoint = GetFocusPoint(InFocus, InputTolerance))
 	{
 		FNodalPointOffset NodalPointOffset;
 		if (NodalOffsetFocusPoint->GetPoint(InZoom, NodalPointOffset, InputTolerance))
diff --git a/Engine/Plugins/VirtualProduction/CameraCalibrationCore/Source/CameraCalibrationCore/Private/Tables/STMapTable.cpp b/Engine/Plugins/VirtualProduction/CameraCalibrationCore/Source/CameraCalibrationCore/Private/Tables/STMapTable.cpp
--- a/Engine/Plugins/VirtualProduction/CameraCalibrationCore/Source/CameraCalibrationCore/Private/Tables/STMapTable.cpp
+++ b/Engine/Plugins/VirtualProduction/CameraCalibrationCore/Source/CameraCalibrationCore/Private/Tables/STMapTable.cpp
@@ -118,9 +118,9 @@ void FSTMapTable::ForEachPoint(FFocusPointCallback InCallback) const
 	}
 }
 
-bool FSTMapTable::DoesFocusPointExists(float InFocus) const
+bool FSTMapTable::DoesFocusPointExists(float InFocus, float InputTolerance) const
 {
-	if (GetFocusPoint(InFocus) != nullptr)
+	if (GetFocusPoint(InFocus, InputTolerance) != nullptr)
 	{
 		return true;
 	}
@@ -170,14 +170,14 @@ bool FSTMapTable::BuildMapBlendingCurve(float InFocus, FRichCurve& OutCurve)
 	return false;
 }
 
-const FSTMapFocusPoint* FSTMapTable::GetFocusPoint(float InFocus) const
+const FSTMapFocusPoint* FSTMapTable::GetFocusPoint(float InFocus, float InputTolerance) const
 {
-	return FocusPoints.FindByPredicate([InFocus](const FSTMapFocusPoint& Point) { return FMath::IsNearlyEqual(Point.Focus, InFocus); });
+	return FocusPoints.FindByPredicate([InFocus, InputTolerance](const FSTMapFocusPoint& Point) { return FMath::IsNearlyEqual(Point.Focus, InFocus, InputTolerance); });
 }
 
-FSTMapFocusPoint* FSTMapTable::GetFocusPoint(float InFocus)
+FSTMapFocusPoint* FSTMapTable::GetFocusPoint(float InFocus, float InputTolerance)
 {
-	return FocusPoints.FindByPredicate([InFocus](const FSTMapFocusPoint& Point) { return FMath::IsNearlyEqual(Point.Focus, InFocus); });
+	return FocusPoints.FindByPredicate([InFocus, InputTolerance](const FSTMapFocusPoint& Point) { return FMath::IsNearlyEqual(Point.Focus, InFocus, InputTolerance); });
 }
 
 TConstArrayView<FSTMapFocusPoint> FSTMapTable::GetFocusPoints() const
@@ -219,7 +219,7 @@ bool FSTMapTable::AddPoint(float InFocus, float InZoom, const FSTMapInfo& InData
 
 bool FSTMapTable::GetPoint(const float InFocus, const float InZoom, FSTMapInfo& OutData, float InputTolerance) const
 {
-	if (const FSTMapFocusPoint* STMapFocusPoint = GetFocusPoint(InFocus))
+	if (const FSTMapFocusPoint* STMapFocusPoint = GetFocusPoint(InFocus, InputTolerance))
 	{
 		FSTMapInfo SMapInfo;
 		if (STMapFocusPoint->GetPoint(InZoom, SMapInfo, InputTolerance))
diff --git a/Engine/Plugins/VirtualProduction/CameraCalibrationCore/Source/CameraCalibrationCore/Public/Tables/BaseLensTable.h b/Engine/Plugins/VirtualProduction/CameraCalibrationCore/Source/CameraCalibrationCore/Public/Tables/BaseLensTable.h
--- a/Engine/Plugins/VirtualProduction/CameraCalibrationCore/Source/CameraCalibrationCore/Public/Tables/BaseLensTable.h
+++ b/Engine/Plugins/VirtualProduction/CameraCalibrationCore/Source/CameraCalibrationCore/Public/Tables/BaseLensTable.h
@@ -82,7 +82,7 @@ protected:
 	 * @param InFocus focus value to check
 	 * @return true if the point exists
 	 */
-	virtual bool DoesFocusPointExists(float InFocus) const PURE_VIRTUAL(FBaseLensTable::DoesFocusPointExists, return false; );
+	virtual bool DoesFocusPointExists(float InFocus, float InputTolerance = KINDA_SMALL_NUMBER) const PURE_VIRTUAL(FBaseLensTable::DoesFocusPointExists, return false; );
 
 	/**
 	 * Whether the zoom point exists
diff --git a/Engine/Plugins/VirtualProduction/CameraCalibrationCore/Source/CameraCalibrationCore/Public/Tables/DistortionParametersTable.h b/Engine/Plugins/VirtualProduction/CameraCalibrationCore/Source/CameraCalibrationCore/Public/Tables/DistortionParametersTable.h
--- a/Engine/Plugins/VirtualProduction/CameraCalibrationCore/Source/CameraCalibrationCore/Public/Tables/DistortionParametersTable.h
+++ b/Engine/Plugins/VirtualProduction/CameraCalibrationCore/Source/CameraCalibrationCore/Public/Tables/DistortionParametersTable.h
@@ -86,7 +86,7 @@ struct CAMERACALIBRATIONCORE_API FDistortionTable : public FBaseLensTable
 protected:
 	//~ Begin FBaseDataTable Interface
 	virtual TMap<ELensDataCategory, FLinkPointMetadata> GetLinkedCategories() const override;
-	virtual bool DoesFocusPointExists(float InFocus) const override;
+	virtual bool DoesFocusPointExists(float InFocus, float InputTolerance = KINDA_SMALL_NUMBER) const override;
 	virtual bool DoesZoomPointExists(float InFocus, float InZoom, float InputTolerance = KINDA_SMALL_NUMBER) const override;
 	virtual const FBaseFocusPoint* GetBaseFocusPoint(int32 InIndex) const override;
 	//~ End FBaseDataTable Interface
@@ -106,10 +106,10 @@ public:
 	bool BuildParameterCurve(float InFocus, int32 ParameterIndex, FRichCurve& OutCurve) const;
 
 	/** Returns const point for a given focus */
-	const FDistortionFocusPoint* GetFocusPoint(float InFocus) const;
+	const FDistortionFocusPoint* GetFocusPoint(float InFocus, float InputTolerance = KINDA_SMALL_NUMBER) const;
 	
 	/** Returns point for a given focus */
-	FDistortionFocusPoint* GetFocusPoint(float InFocus);
+	FDistortionFocusPoint* GetFocusPoint(float InFocus, float InputTolerance = KINDA_SMALL_NUMBER);
 
 	/** Returns all focus points */
 	TConstArrayView<FDistortionFocusPoint> GetFocusPoints() const;
diff --git a/Engine/Plugins/VirtualProduction/CameraCalibrationCore/Source/CameraCalibrationCore/Public/Tables/FocalLengthTable.h b/Engine/Plugins/VirtualProduction/CameraCalibrationCore/Source/CameraCalibrationCore/Public/Tables/FocalLengthTable.h
--- a/Engine/Plugins/VirtualProduction/CameraCalibrationCore/Source/CameraCalibrationCore/Public/Tables/FocalLengthTable.h
+++ b/Engine/Plugins/VirtualProduction/CameraCalibrationCore/Source/CameraCalibrationCore/Public/Tables/FocalLengthTable.h
@@ -97,7 +97,7 @@ struct CAMERACALIBRATIONCORE_API FFocalLengthTable : public FBaseLensTable
 protected:
 	//~ Begin FBaseDataTable Interface
 	virtual TMap<ELensDataCategory, FLinkPointMetadata> GetLinkedCategories() const override;
-	virtual bool DoesFocusPointExists(float InFocus) const override;
+	virtual bool DoesFocusPointExists(float InFocus, float InputTolerance = KINDA_SMALL_NUMBER) const override;
 	virtual bool DoesZoomPointExists(float InFocus, float InZoom, float InputTolerance = KINDA_SMALL_NUMBER) const override;
 	virtual const FBaseFocusPoint* GetBaseFocusPoint(int32 InIndex) const override;
 	//~ End FBaseDataTable Interface
@@ -117,10 +117,10 @@ public:
 	bool BuildParameterCurve(float InFocus, int32 ParameterIndex, FRichCurve& OutCurve) const;
 	
 	/** Returns const point for a given focus */
-	const FFocalLengthFocusPoint* GetFocusPoint(float InFocus) const;
+	const FFocalLengthFocusPoint* GetFocusPoint(float InFocus, float InputTolerance = KINDA_SMALL_NUMBER) const;
 	
 	/** Returns point for a given focus */
-	FFocalLengthFocusPoint* GetFocusPoint(float InFocus);
+	FFocalLengthFocusPoint* GetFocusPoint(float InFocus, float InputTolerance = KINDA_SMALL_NUMBER);
 
 	/** Returns all focus points */
 	TConstArrayView<FFocalLengthFocusPoint> GetFocusPoints() const;
diff --git a/Engine/Plugins/VirtualProduction/CameraCalibrationCore/Source/CameraCalibrationCore/Public/Tables/ImageCenterTable.h b/Engine/Plugins/VirtualProduction/CameraCalibrationCore/Source/CameraCalibrationCore/Public/Tables/ImageCenterTable.h
--- a/Engine/Plugins/VirtualProduction/CameraCalibrationCore/Source/CameraCalibrationCore/Public/Tables/ImageCenterTable.h
+++ b/Engine/Plugins/VirtualProduction/CameraCalibrationCore/Source/CameraCalibrationCore/Public/Tables/ImageCenterTable.h
@@ -67,7 +67,7 @@ struct CAMERACALIBRATIONCORE_API FImageCenterTable : public FBaseLensTable
 protected:
 	//~ Begin FBaseDataTable Interface
 	virtual TMap<ELensDataCategory, FLinkPointMetadata> GetLinkedCategories() const override;
-	virtual bool DoesFocusPointExists(float InFocus) const override;
+	virtual bool DoesFocusPointExists(float InFocus, float InputTolerance = KINDA_SMALL_NUMBER) const override;
 	virtual bool DoesZoomPointExists(float InFocus, float InZoom, float InputTolerance = KINDA_SMALL_NUMBER) const override;
 	virtual const FBaseFocusPoint* GetBaseFocusPoint(int32 InIndex) const override;
 	//~ End FBaseDataTable Interface
@@ -87,10 +87,10 @@ public:
 	bool BuildParameterCurve(float InFocus, int32 ParameterIndex, FRichCurve& OutCurve) const;
 	
 	/** Returns const point for a given focus */
-	const FImageCenterFocusPoint* GetFocusPoint(float InFocus) const;
+	const FImageCenterFocusPoint* GetFocusPoint(float InFocus, float InputTolerance = KINDA_SMALL_NUMBER) const;
 
 	/** Returns const point for a given focus */
-	FImageCenterFocusPoint* GetFocusPoint(float InFocus);
+	FImageCenterFocusPoint* GetFocusPoint(float InFocus, float InputTolerance = KINDA_SMALL_NUMBER);
 
 	/** Returns all focus points */
 	TConstArrayView<FImageCenterFocusPoint> GetFocusPoints() const;
diff --git a/Engine/Plugins/VirtualProduction/CameraCalibrationCore/Source/CameraCalibrationCore/Public/Tables/NodalOffsetTable.h b/Engine/Plugins/VirtualProduction/CameraCalibrationCore/Source/CameraCalibrationCore/Public/Tables/NodalOffsetTable.h
--- a/Engine/Plugins/VirtualProduction/CameraCalibrationCore/Source/CameraCalibrationCore/Public/Tables/NodalOffsetTable.h
+++ b/Engine/Plugins/VirtualProduction/CameraCalibrationCore/Source/CameraCalibrationCore/Public/Tables/NodalOffsetTable.h
@@ -73,7 +73,7 @@ struct CAMERACALIBRATIONCORE_API FNodalOffsetTable : public FBaseLensTable
 protected:
 	//~ Begin FBaseDataTable Interface
 	virtual TMap<ELensDataCategory, FLinkPointMetadata> GetLinkedCategories() const override;
-	virtual bool DoesFocusPointExists(float InFocus) const override;
+	virtual bool DoesFocusPointExists(float InFocus, float InputTolerance = KINDA_SMALL_NUMBER) const override;
 	virtual bool DoesZoomPointExists(float InFocus, float InZoom, float InputTolerance = KINDA_SMALL_NUMBER) const override;
 	virtual const FBaseFocusPoint* GetBaseFocusPoint(int32 InIndex) const override;
 	//~ End FBaseDataTable Interface
@@ -93,10 +93,10 @@ public:
 	bool BuildParameterCurve(float InFocus, int32 ParameterIndex, EAxis::Type InAxis, FRichCurve& OutCurve) const;
 
 	/** Returns const point for a given focus */
-	const FNodalOffsetFocusPoint* GetFocusPoint(float InFocus) const;
+	const FNodalOffsetFocusPoint* GetFocusPoint(float InFocus, float InputTolerance = KINDA_SMALL_NUMBER) const;
 
 	/** Returns point for a given focus */
-	FNodalOffsetFocusPoint* GetFocusPoint(float InFocus);
+	FNodalOffsetFocusPoint* GetFocusPoint(float InFocus, float InputTolerance = KINDA_SMALL_NUMBER);
 
 	/** Returns all focus points */
 	TConstArrayView<FNodalOffsetFocusPoint> GetFocusPoints() const;
diff --git a/Engine/Plugins/VirtualProduction/CameraCalibrationCore/Source/CameraCalibrationCore/Public/Tables/STMapTable.h b/Engine/Plugins/VirtualProduction/CameraCalibrationCore/Source/CameraCalibrationCore/Public/Tables/STMapTable.h
--- a/Engine/Plugins/VirtualProduction/CameraCalibrationCore/Source/CameraCalibrationCore/Public/Tables/STMapTable.h
+++ b/Engine/Plugins/VirtualProduction/CameraCalibrationCore/Source/CameraCalibrationCore/Public/Tables/STMapTable.h
@@ -125,7 +125,7 @@ struct CAMERACALIBRATIONCORE_API FSTMapTable : public FBaseLensTable
 protected:
 	//~ Begin FBaseDataTable Interface
 	virtual TMap<ELensDataCategory, FLinkPointMetadata> GetLinkedCategories() const override;
-	virtual bool DoesFocusPointExists(float InFocus) const override;
+	virtual bool DoesFocusPointExists(float InFocus, float InputTolerance = KINDA_SMALL_NUMBER) const override;
 	virtual bool DoesZoomPointExists(float InFocus, float InZoom, float InputTolerance = KINDA_SMALL_NUMBER) const override;
 	virtual const FBaseFocusPoint* GetBaseFocusPoint(int32 InIndex) const override;
 	//~ End FBaseDataTable Interface
@@ -145,10 +145,10 @@ public:
 	bool BuildMapBlendingCurve(float InFocus, FRichCurve& OutCurve);
 	
 	/** Returns const point for a given focus */
-	const FSTMapFocusPoint* GetFocusPoint(float InFocus) const;
+	const FSTMapFocusPoint* GetFocusPoint(float InFocus, float InputTolerance = KINDA_SMALL_NUMBER) const;
 
 	/** Returns point for a given focus */
-	FSTMapFocusPoint* GetFocusPoint(float InFocus);
+	FSTMapFocusPoint* GetFocusPoint(float InFocus, float InputTolerance = KINDA_SMALL_NUMBER);
 
 	/** Returns all focus points */
 	TConstArrayView<FSTMapFocusPoint> GetFocusPoints() const;
