CameraCalibration - Add Squeeze Factor to LensFile info for use with anamorphic lenses
#rnx

diff --git a/Engine/Plugins/Compositing/Composure/Source/Composure/Private/ComposureBlueprintLibrary.cpp b/Engine/Plugins/Compositing/Composure/Source/Composure/Private/ComposureBlueprintLibrary.cpp
--- a/Engine/Plugins/Compositing/Composure/Source/Composure/Private/ComposureBlueprintLibrary.cpp
+++ b/Engine/Plugins/Compositing/Composure/Source/Composure/Private/ComposureBlueprintLibrary.cpp
@@ -92,7 +92,7 @@ void UComposureBlueprintLibrary::CopyCameraSettingsToSceneCapture(UCameraCompone
 			}
 			else
 			{
-				const float OverscanSensorWidth = SrcCineCameraComponent->Filmback.SensorWidth * OverscanFactor;
+				const float OverscanSensorWidth = SrcCineCameraComponent->Filmback.SensorWidth * SrcCineCameraComponent->LensSettings.SqueezeFactor * OverscanFactor;
 				Dst->FOVAngle = FMath::RadiansToDegrees(2.0f * FMath::Atan(OverscanSensorWidth / (2.0f * OriginalFocalLength)));
 			}
 		}
diff --git a/Engine/Plugins/VirtualProduction/CameraCalibration/Source/CameraCalibrationEditor/Private/AssetEditor/CameraCalibrationStepsController.cpp b/Engine/Plugins/VirtualProduction/CameraCalibration/Source/CameraCalibrationEditor/Private/AssetEditor/CameraCalibrationStepsController.cpp
--- a/Engine/Plugins/VirtualProduction/CameraCalibration/Source/CameraCalibrationEditor/Private/AssetEditor/CameraCalibrationStepsController.cpp
+++ b/Engine/Plugins/VirtualProduction/CameraCalibration/Source/CameraCalibrationEditor/Private/AssetEditor/CameraCalibrationStepsController.cpp
@@ -281,7 +281,7 @@ void FCameraCalibrationStepsController::UpdateCGResolution(bool bCompResized)
 
 	const FIntPoint CGResolution = CGLayerPtr->GetRenderResolution();
 	const float CGLayerAspectRatio = (CGResolution.Y != 0) ? (float)CGResolution.X / (float)CGResolution.Y : 0.0f;
-	const float FilmbackAspectRatio = CineCameraComponentPtr->Filmback.SensorAspectRatio;
+	const float FilmbackAspectRatio = CineCameraComponentPtr->AspectRatio;
 
 	// Early-out if the filmback is somehow 0.0f (which shouldn't be possible)
 	if (FMath::IsNearlyEqual(FilmbackAspectRatio, 0.0f))
@@ -348,7 +348,7 @@ void FCameraCalibrationStepsController::SetCameraFeedDimensionsFromMousePosition
 		if (CineCameraComponent.IsValid())
 		{
 			const float CameraFeedAspectRatio = CameraFeedDimensions.X / (float)CameraFeedDimensions.Y;
-			const float CameraAspectRatio = CineCameraComponent->Filmback.SensorAspectRatio;
+			const float CameraAspectRatio = CineCameraComponent->AspectRatio;
 
 			// If the two aspect ratios are within the acceptable tolerance, attempt to minimize the aspect ratio difference by adjusting the camera feed dimensions
 			constexpr float AspectRatioNudgeTolerance = 0.1f;
diff --git a/Engine/Plugins/VirtualProduction/CameraCalibration/Source/CameraCalibrationEditor/Private/AssetEditor/LensInfoStep.cpp b/Engine/Plugins/VirtualProduction/CameraCalibration/Source/CameraCalibrationEditor/Private/AssetEditor/LensInfoStep.cpp
--- a/Engine/Plugins/VirtualProduction/CameraCalibration/Source/CameraCalibrationEditor/Private/AssetEditor/LensInfoStep.cpp
+++ b/Engine/Plugins/VirtualProduction/CameraCalibration/Source/CameraCalibrationEditor/Private/AssetEditor/LensInfoStep.cpp
@@ -2,6 +2,7 @@
 
 #include "LensInfoStep.h"
 
+#include "Camera/CameraActor.h"
 #include "CameraCalibrationStepsController.h"
 #include "IStructureDetailsView.h"
 #include "LensFile.h"
@@ -76,6 +77,18 @@ TSharedRef<SWidget> ULensInfoStep::BuildUI()
 			]
 		]
 
+		+ SVerticalBox::Slot() // Text Warning that the camera info does not match the cine camera's settings
+		.AutoHeight()
+		.Padding(0, 10)
+		[
+			SNew(STextBlock)
+			.Text(LOCTEXT("CameraSettingMismatch", "These filmback settings do not match the settings of the CineCamera.\nIf they are expected to match, set the FilmbackOverride setting on the LensComponent to \"LensFile\"."))
+			.ColorAndOpacity(FSlateColor(FLinearColor::Yellow))
+			.Justification(ETextJustify::Center)
+			.AutoWrapText(true)
+			.Visibility_UObject(this, &ULensInfoStep::GetFilmbackWarningVisibility)
+		]
+
 		+ SVerticalBox::Slot() // Camera Feed information structure
 		.AutoHeight()
 		[
@@ -314,6 +327,39 @@ EVisibility ULensInfoStep::HandleAspectRatioWarningVisibility() const
 	return EVisibility::Collapsed;
 }
 
+EVisibility ULensInfoStep::GetFilmbackWarningVisibility() const
+{
+	if (TSharedPtr<FCameraCalibrationStepsController> StepsController = CameraCalibrationStepsController.Pin())
+	{
+		ULensFile* LensFile = StepsController->GetLensFile();
+		UCineCameraComponent* CineCameraComponent = nullptr;
+
+		if (ACameraActor* Camera = StepsController->GetCamera())
+		{
+			TInlineComponentArray<UCineCameraComponent*> CameraComponents;
+			Camera->GetComponents(CameraComponents);
+
+			if (CameraComponents.Num() > 0)
+			{
+				CineCameraComponent = CameraComponents[0];
+			}
+		}
+
+		if (LensFile && CineCameraComponent)
+		{
+			const FVector2D LensFileFilmback = LensFile->LensInfo.SensorDimensions;
+			if ((LensFile->LensInfo.SensorDimensions.X != CineCameraComponent->Filmback.SensorWidth) ||
+				(LensFile->LensInfo.SensorDimensions.Y != CineCameraComponent->Filmback.SensorHeight) ||
+				(LensFile->LensInfo.SqueezeFactor != CineCameraComponent->LensSettings.SqueezeFactor))
+			{
+				return EVisibility::Visible;
+			}
+		}
+	}
+
+	return EVisibility::Collapsed;
+}
+
 bool ULensInfoStep::IsActive() const
 {
 	return bIsActive;
diff --git a/Engine/Plugins/VirtualProduction/CameraCalibration/Source/CameraCalibrationEditor/Private/AssetEditor/LensInfoStep.h b/Engine/Plugins/VirtualProduction/CameraCalibration/Source/CameraCalibrationEditor/Private/AssetEditor/LensInfoStep.h
--- a/Engine/Plugins/VirtualProduction/CameraCalibration/Source/CameraCalibrationEditor/Private/AssetEditor/LensInfoStep.h
+++ b/Engine/Plugins/VirtualProduction/CameraCalibration/Source/CameraCalibrationEditor/Private/AssetEditor/LensInfoStep.h
@@ -57,6 +57,9 @@ private:
 	/** Determines the visibility of the aspect ratio warning by comparing the camera feed aspect ratio with that of the camera actor */
 	EVisibility HandleAspectRatioWarningVisibility() const;
 
+	/** Determines the visibility of the camera settings warning by comparing the filmback and squeeze factor in the LensFile with that of the camera actor */
+	EVisibility GetFilmbackWarningVisibility() const;
+
 private:
 
 	/** Pointer to the calibration steps controller */
diff --git a/Engine/Plugins/VirtualProduction/CameraCalibration/Source/CameraCalibrationEditor/Private/AssetEditor/SLensEvaluation.cpp b/Engine/Plugins/VirtualProduction/CameraCalibration/Source/CameraCalibrationEditor/Private/AssetEditor/SLensEvaluation.cpp
--- a/Engine/Plugins/VirtualProduction/CameraCalibration/Source/CameraCalibrationEditor/Private/AssetEditor/SLensEvaluation.cpp
+++ b/Engine/Plugins/VirtualProduction/CameraCalibration/Source/CameraCalibrationEditor/Private/AssetEditor/SLensEvaluation.cpp
@@ -137,7 +137,7 @@ void SLensEvaluation::CacheLensFileEvaluationInputs()
 		FFocalLengthInfo FocalLength;
 		if (LensFile->EvaluateFocalLength(EvalInputs.Focus, EvalInputs.Zoom, FocalLength))
 		{
-			CachedFIZData.EvaluatedZoom = FocalLength.FxFy.X * LensFile->LensInfo.SensorDimensions.X;
+			CachedFIZData.EvaluatedZoom = FocalLength.FxFy.X * LensFile->LensInfo.SensorDimensions.X * LensFile->LensInfo.SqueezeFactor;
 		}
 	}
 }
@@ -497,7 +497,8 @@ TSharedRef<SWidget> SLensEvaluation::MakeEvaluatedFIZWidget() const
 				{
 					if (CachedFIZData.EvaluatedZoom.IsSet())
 					{
-						const float FOV = FMath::RadiansToDegrees(2.f * FMath::Atan(LensFile->LensInfo.SensorDimensions.X / (2.f * CachedFIZData.EvaluatedZoom.GetValue())));
+						const float DesqueezedSensorWidth = LensFile->LensInfo.SensorDimensions.X * LensFile->LensInfo.SqueezeFactor;
+						const float FOV = FMath::RadiansToDegrees(2.f * FMath::Atan(DesqueezedSensorWidth / (2.f * CachedFIZData.EvaluatedZoom.GetValue())));
 						const FText ValueText = FText::AsNumber(FOV, &FloatOptions);
 						return FText::Format(LOCTEXT("PhysicalUnitsFOVValue", "{0} deg"), ValueText);
 					}
diff --git a/Engine/Plugins/VirtualProduction/CameraCalibration/Source/CameraCalibrationEditor/Private/Calibrators/CameraLensDistortionAlgoCheckerboard.cpp b/Engine/Plugins/VirtualProduction/CameraCalibration/Source/CameraCalibrationEditor/Private/Calibrators/CameraLensDistortionAlgoCheckerboard.cpp
--- a/Engine/Plugins/VirtualProduction/CameraCalibration/Source/CameraCalibrationEditor/Private/Calibrators/CameraLensDistortionAlgoCheckerboard.cpp
+++ b/Engine/Plugins/VirtualProduction/CameraCalibration/Source/CameraCalibrationEditor/Private/Calibrators/CameraLensDistortionAlgoCheckerboard.cpp
@@ -448,12 +448,6 @@ TSharedRef<SWidget> UCameraLensDistortionAlgoCheckerboard::BuildUI()
 		.MaxHeight(FCameraCalibrationWidgetHelpers::DefaultRowHeight)
 		[FCameraCalibrationWidgetHelpers::BuildLabelWidgetPair(LOCTEXT("ShowDetection", "Show Detection"), BuildShowDetectionWidget())]
 
-	+ SVerticalBox::Slot() // Pixel Aspect
-		.VAlign(EVerticalAlignment::VAlign_Top)
-		.AutoHeight()
-		.MaxHeight(FCameraCalibrationWidgetHelpers::DefaultRowHeight)
-		[FCameraCalibrationWidgetHelpers::BuildLabelWidgetPair(LOCTEXT("PixelAspect", "Pixel Aspect"), BuildPixelAspectWidget())]
-
 	+ SVerticalBox::Slot() // Fix Focal Length
 		.VAlign(EVerticalAlignment::VAlign_Top)
 		.AutoHeight()
@@ -643,6 +637,8 @@ bool UCameraLensDistortionAlgoCheckerboard::GetLensDistortion(
 	float TestFocalLength = StepsController->GetLensFileEvaluationInputs().Zoom;
 	float TestSensorWidth = StepsController->GetLensFileEvaluationInputs().Filmback.SensorWidth;
 
+	const float PixelAspect = LensFile->LensInfo.SqueezeFactor;
+
 	if (CVarFixFocalLengthValue.GetValueOnGameThread() > 1.0f)
 	{
 		TestFocalLength = CVarFixFocalLengthValue.GetValueOnGameThread();
@@ -660,6 +656,8 @@ bool UCameraLensDistortionAlgoCheckerboard::GetLensDistortion(
 		TestImageHeight = CVarImageHeightValue.GetValueOnGameThread();
 	}
 
+	TestSensorWidth = TestSensorWidth * PixelAspect;
+
 	float Fx = TestImageWidth;
 	if (!FMath::IsNearlyZero(TestSensorWidth))
 	{
@@ -828,17 +826,6 @@ TSharedRef<SWidget> UCameraLensDistortionAlgoCheckerboard::BuildShowDetectionWid
 	});
 }
 
-TSharedRef<SWidget> UCameraLensDistortionAlgoCheckerboard::BuildPixelAspectWidget()
-{
-	return SNew(SNumericEntryBox<float>)
-		.Value_Lambda([&]() { return PixelAspect; })
-		.ToolTipText(LOCTEXT("Pixel Aspect", "Pixel Aspect"))
-		.OnValueChanged_Lambda([&](double InValue)
-		{
-			PixelAspect = InValue;
-		});
-}
-
 TSharedRef<SWidget> UCameraLensDistortionAlgoCheckerboard::BuildFixFocalLengthWidget()
 {
 	return SNew(SCheckBox)
diff --git a/Engine/Plugins/VirtualProduction/CameraCalibration/Source/CameraCalibrationEditor/Private/Calibrators/CameraLensDistortionAlgoCheckerboard.h b/Engine/Plugins/VirtualProduction/CameraCalibration/Source/CameraCalibrationEditor/Private/Calibrators/CameraLensDistortionAlgoCheckerboard.h
--- a/Engine/Plugins/VirtualProduction/CameraCalibration/Source/CameraCalibrationEditor/Private/Calibrators/CameraLensDistortionAlgoCheckerboard.h
+++ b/Engine/Plugins/VirtualProduction/CameraCalibration/Source/CameraCalibrationEditor/Private/Calibrators/CameraLensDistortionAlgoCheckerboard.h
@@ -186,9 +186,6 @@ private:
 	/** True if a detection window should be shown after every capture */
 	bool bShouldShowDetectionWindow = false;
 
-	/** Pixel Aspect value to pass to the solver */
-	float PixelAspect = 1.0f;
-
 	/** If true, the solver will not solve for focal length, and will just use the input value */
 	bool bFixFocalLength = false;
 
@@ -214,9 +211,6 @@ private:
 	/** Builds the UI for the user to select if they want a corner detection window to be shown after every capture */
 	TSharedRef<SWidget> BuildShowDetectionWidget();
 
-	/** Builds the UI of the pixel aspect widget */
-	TSharedRef<SWidget> BuildPixelAspectWidget();
-
 	/** Builds the UI of the fix focal length width */
 	TSharedRef<SWidget> BuildFixFocalLengthWidget();
 
diff --git a/Engine/Plugins/VirtualProduction/CameraCalibrationCore/Source/CameraCalibrationCore/Private/AnamorphicLensDistortionModelHandler.cpp b/Engine/Plugins/VirtualProduction/CameraCalibrationCore/Source/CameraCalibrationCore/Private/AnamorphicLensDistortionModelHandler.cpp
--- a/Engine/Plugins/VirtualProduction/CameraCalibrationCore/Source/CameraCalibrationCore/Private/AnamorphicLensDistortionModelHandler.cpp
+++ b/Engine/Plugins/VirtualProduction/CameraCalibrationCore/Source/CameraCalibrationCore/Private/AnamorphicLensDistortionModelHandler.cpp
@@ -15,7 +15,7 @@ FVector2D UAnamorphicLensDistortionModelHandler::ComputeDistortedUV(const FVecto
 	// The following implementation is based on the 3DE4 Anamorphic - Standard, Degree 4 model for distortion
 
 	// Pre-compute some values that will be useful later
-	const FVector2D FilmbackSize = FVector2D(CameraFilmback.SensorWidth, CameraFilmback.SensorHeight);
+	const FVector2D FilmbackSize = FVector2D(CameraFilmback.SensorWidth * AnamorphicParameters.PixelAspect, CameraFilmback.SensorHeight);
 
 	// Guard against divide-by-zero errors
 	if (FMath::IsNearlyZero(FilmbackSize.X) || FMath::IsNearlyZero(FilmbackSize.Y) || FMath::IsNearlyZero(AnamorphicParameters.PixelAspect))
@@ -95,7 +95,7 @@ void UAnamorphicLensDistortionModelHandler::UpdateMaterialParameters()
 	{
 		MID->SetScalarParameterValue("pixel_aspect", AnamorphicParameters.PixelAspect);
 
- 		MID->SetScalarParameterValue("w_fb", CameraFilmback.SensorWidth);
+ 		MID->SetScalarParameterValue("w_fb", CameraFilmback.SensorWidth * AnamorphicParameters.PixelAspect);
  		MID->SetScalarParameterValue("h_fb", CameraFilmback.SensorHeight);
 
 		MID->SetScalarParameterValue("cx02", AnamorphicParameters.CX02);
diff --git a/Engine/Plugins/VirtualProduction/CameraCalibrationCore/Source/CameraCalibrationCore/Private/LensComponent.cpp b/Engine/Plugins/VirtualProduction/CameraCalibrationCore/Source/CameraCalibrationCore/Private/LensComponent.cpp
--- a/Engine/Plugins/VirtualProduction/CameraCalibrationCore/Source/CameraCalibrationCore/Private/LensComponent.cpp
+++ b/Engine/Plugins/VirtualProduction/CameraCalibrationCore/Source/CameraCalibrationCore/Private/LensComponent.cpp
@@ -186,7 +186,7 @@ void ULensComponent::TickComponent(float DeltaTime, ELevelTick TickType, FActorC
 
 				// Get the overscan factor and use it to modify the target camera's FOV
 				const float OverscanFactor = LensDistortionHandler->GetOverscanFactor();
-				const float OverscanSensorWidth = CineCameraComponent->Filmback.SensorWidth * OverscanFactor;
+				const float OverscanSensorWidth = GetDesqueezedSensorWidth(CineCameraComponent) * OverscanFactor;
 				const float OverscanFOV = FMath::RadiansToDegrees(2.0f * FMath::Atan(OverscanSensorWidth / (2.0f * OriginalFocalLength)));
 				CineCameraComponent->SetFieldOfView(OverscanFOV);
 
@@ -451,7 +451,7 @@ void ULensComponent::EvaluateFocalLength(UCineCameraComponent* CineCameraCompone
 			// => FilmbackY = FocalLength / Fy
 
 			// Adjust FocalLength and Filmback to match FxFy (which has already been divided by resolution in pixels)
-			const float NewFocalLength = FocalLengthInfo.FxFy.X * EvalInputs.Filmback.SensorWidth;
+			const float NewFocalLength = FocalLengthInfo.FxFy.X * EvalInputs.Filmback.SensorWidth * LensFile->LensInfo.SqueezeFactor;
 
 			// TODO: Consider adding an advanced setting to control whether or not sensor height should be modified (if Fx/Fy don't have a fixed aspect ratio)
 
@@ -696,6 +696,16 @@ void ULensComponent::InitDefaultCamera()
 	}
 }
 
+float ULensComponent::GetDesqueezedSensorWidth(UCineCameraComponent* const CineCameraComponent) const
+{
+	if (CineCameraComponent)
+	{
+		return CineCameraComponent->Filmback.SensorWidth * CineCameraComponent->LensSettings.SqueezeFactor;
+	}
+
+	return 1.0f;
+}
+
 void ULensComponent::CleanupDistortion(UCineCameraComponent* const CineCameraComponent)
 {
 	if (bIsDistortionSetup)
@@ -708,7 +718,7 @@ void ULensComponent::CleanupDistortion(UCineCameraComponent* const CineCameraCom
 		}
 
 		// Restore the original FOV of the target camera
-		const float UndistortedFOV = FMath::RadiansToDegrees(2.0f * FMath::Atan(CineCameraComponent->Filmback.SensorWidth / (2.0f * OriginalFocalLength)));
+		const float UndistortedFOV = FMath::RadiansToDegrees(2.0f * FMath::Atan(GetDesqueezedSensorWidth(CineCameraComponent) / (2.0f * OriginalFocalLength)));
 		CineCameraComponent->SetFieldOfView(UndistortedFOV);
 
 		// Update the minimum and maximum focal length of the camera (if needed)
@@ -875,12 +885,18 @@ void ULensComponent::UpdateCameraFilmback(UCineCameraComponent* CineCameraCompon
 			{
 				CineCameraComponent->Filmback.SensorWidth = LensFile->LensInfo.SensorDimensions.X;
 				CineCameraComponent->Filmback.SensorHeight = LensFile->LensInfo.SensorDimensions.Y;
+				CineCameraComponent->LensSettings.SqueezeFactor = LensFile->LensInfo.SqueezeFactor;
 			}
 			break;
 		}
 		case EFilmbackOverrideSource::CroppedFilmbackSetting:
 		{
 			CineCameraComponent->Filmback = CroppedFilmback;
+
+			if (ULensFile* LensFile = LensFilePicker.GetLensFile())
+			{
+				CineCameraComponent->LensSettings.SqueezeFactor = LensFile->LensInfo.SqueezeFactor;
+			}
 			break;
 		}
 		case EFilmbackOverrideSource::DoNotOverride:
diff --git a/Engine/Plugins/VirtualProduction/CameraCalibrationCore/Source/CameraCalibrationCore/Public/LensComponent.h b/Engine/Plugins/VirtualProduction/CameraCalibrationCore/Source/CameraCalibrationCore/Public/LensComponent.h
--- a/Engine/Plugins/VirtualProduction/CameraCalibrationCore/Source/CameraCalibrationCore/Public/LensComponent.h
+++ b/Engine/Plugins/VirtualProduction/CameraCalibrationCore/Source/CameraCalibrationCore/Public/LensComponent.h
@@ -254,6 +254,9 @@ private:
 	/** Updates the camera's filmback based on the filmback override settings */
 	void UpdateCameraFilmback(UCineCameraComponent* CineCameraComponent);
 
+	/** Returns the sensor width of the input CineCamera, factoring in its squeeze factor */
+	float GetDesqueezedSensorWidth(UCineCameraComponent* const CineCameraComponent) const;
+
 protected:
 	/** Lens File used to drive distortion with current camera settings */
 	UPROPERTY(EditAnywhere, Category="Lens File", meta=(ShowOnlyInnerProperties))
diff --git a/Engine/Plugins/VirtualProduction/CameraCalibrationCore/Source/CameraCalibrationCore/Public/LensData.h b/Engine/Plugins/VirtualProduction/CameraCalibrationCore/Source/CameraCalibrationCore/Public/LensData.h
--- a/Engine/Plugins/VirtualProduction/CameraCalibrationCore/Source/CameraCalibrationCore/Public/LensData.h
+++ b/Engine/Plugins/VirtualProduction/CameraCalibrationCore/Source/CameraCalibrationCore/Public/LensData.h
@@ -58,6 +58,9 @@ public:
 	UPROPERTY(EditAnywhere, BlueprintReadWrite, Category = "Camera Info", meta = (DisplayName = "Image Resolution"))
 	FIntPoint ImageDimensions = FIntPoint(1920, 1080);
 
+	/** Squeeze Factor (or Pixel Aspect) for anamorphic lenses. Spherical Lenses should keep this default to 1.0f */
+	UPROPERTY(EditAnywhere, BlueprintReadWrite, Category = "Camera Info")
+	float SqueezeFactor = 1.0f;
 };
 
 /**
