CameraCalibration - Fix crash when copying LensComponent

diff --git a/Engine/Plugins/VirtualProduction/CameraCalibrationCore/Source/CameraCalibrationCore/Private/LensComponent.cpp b/Engine/Plugins/VirtualProduction/CameraCalibrationCore/Source/CameraCalibrationCore/Private/LensComponent.cpp
--- a/Engine/Plugins/VirtualProduction/CameraCalibrationCore/Source/CameraCalibrationCore/Private/LensComponent.cpp
+++ b/Engine/Plugins/VirtualProduction/CameraCalibrationCore/Source/CameraCalibrationCore/Private/LensComponent.cpp
@@ -47,7 +47,10 @@ void ULensComponent::OnRegister()
 
 	// Look for any LiveLinkComponents that were previously added to this component's parent actor
 	TInlineComponentArray<ULiveLinkComponentController*> LiveLinkComponents;
-	GetOwner()->GetComponents(LiveLinkComponents);
+	if (AActor* Owner = GetOwner())
+	{
+		Owner->GetComponents(LiveLinkComponents);
+	}
 
 	for (ULiveLinkComponentController* LiveLinkComponent : LiveLinkComponents)
 	{	
@@ -68,7 +71,10 @@ void ULensComponent::OnUnregister()
 	LiveLinkComponentsModule.OnLiveLinkComponentRegistered().RemoveAll(this);
 
 	TInlineComponentArray<ULiveLinkComponentController*> LiveLinkComponents;
-	GetOwner()->GetComponents(LiveLinkComponents);
+	if (AActor* Owner = GetOwner())
+	{
+		Owner->GetComponents(LiveLinkComponents);
+	}
 
 	for (ULiveLinkComponentController* LiveLinkComponent : LiveLinkComponents)
 	{
@@ -334,7 +340,10 @@ void ULensComponent::ReapplyNodalOffset()
 {
 	// Find the TrackedComponent based on the serialized name
 	TInlineComponentArray<USceneComponent*> SceneComponents;
-	GetOwner()->GetComponents(SceneComponents);
+	if (AActor* Owner = GetOwner())
+	{
+		Owner->GetComponents(SceneComponents);
+	}
 
 	for (USceneComponent* SceneComponent : SceneComponents)
 	{
@@ -687,7 +696,11 @@ void ULensComponent::InitDefaultCamera()
 	{
 		// Find all CineCameraComponents on the same actor as this component and set the first one to be the target
 		TInlineComponentArray<UCineCameraComponent*> CineCameraComponents;
-		GetOwner()->GetComponents(CineCameraComponents);
+		if (AActor* Owner = GetOwner())
+		{
+			Owner->GetComponents(CineCameraComponents);
+		}
+
 		if (CineCameraComponents.Num() > 0)
 		{
 			TargetCameraComponent.ComponentProperty = CineCameraComponents[0]->GetFName();
