CameraCalibration: Fix checkerboard actor's root component to be editable, so it can be driven by LiveLink

diff --git a/Engine/Plugins/VirtualProduction/CameraCalibrationCore/Source/CameraCalibrationCore/Public/CameraCalibrationCheckerboard.h b/Engine/Plugins/VirtualProduction/CameraCalibrationCore/Source/CameraCalibrationCore/Public/CameraCalibrationCheckerboard.h
--- a/Engine/Plugins/VirtualProduction/CameraCalibrationCore/Source/CameraCalibrationCore/Public/CameraCalibrationCheckerboard.h
+++ b/Engine/Plugins/VirtualProduction/CameraCalibrationCore/Source/CameraCalibrationCore/Public/CameraCalibrationCheckerboard.h
@@ -40,7 +40,7 @@ public:
 public:
 
 	/** Root component, gives the Actor a transform */
-	UPROPERTY()
+	UPROPERTY(EditAnywhere, Category = "Calibration")
 	TObjectPtr<USceneComponent> Root;
 
 	/** TopLeft calibration point */
