CameraCalibration: LiveLink controllers can now be added.

Also added default even/odd materials.

This implementation could probably be simplified to a single plane with a checkerboard material.

#virtualized

diff --git a/Engine/Plugins/VirtualProduction/CameraCalibrationCore/Source/CameraCalibrationCore/Private/CameraCalibrationCheckerboard.cpp b/Engine/Plugins/VirtualProduction/CameraCalibrationCore/Source/CameraCalibrationCore/Private/CameraCalibrationCheckerboard.cpp
--- a/Engine/Plugins/VirtualProduction/CameraCalibrationCore/Source/CameraCalibrationCore/Private/CameraCalibrationCheckerboard.cpp
+++ b/Engine/Plugins/VirtualProduction/CameraCalibrationCore/Source/CameraCalibrationCore/Private/CameraCalibrationCheckerboard.cpp
@@ -35,6 +35,29 @@ ACameraCalibrationCheckerboard::ACameraCalibrationCheckerboard() : AActor()
 			CubeMesh = CubeFinder.Object;
 		}
 	}
+
+	// Find default material for odd squares
+	if (!OddCubeMaterial)
+	{
+		ConstructorHelpers::FObjectFinder<UMaterialInterface> MaterialFinder(TEXT("Material'/CameraCalibrationCore/Materials/M_DefaultOddSquare.M_DefaultOddSquare'"), LOAD_Quiet | LOAD_NoWarn);
+		
+		if (MaterialFinder.Succeeded())
+		{
+			OddCubeMaterial = MaterialFinder.Object;
+		}
+	}
+
+	// Find default material for even squares
+	if (!EvenCubeMaterial)
+	{
+		ConstructorHelpers::FObjectFinder<UMaterialInterface> MaterialFinder(TEXT("Material'/CameraCalibrationCore/Materials/M_DefaultEvenSquare.M_DefaultEvenSquare'"), LOAD_Quiet | LOAD_NoWarn);
+
+		if (MaterialFinder.Succeeded())
+		{
+			EvenCubeMaterial = MaterialFinder.Object;
+		}
+	}
+
 }
 
 #if WITH_EDITOR
@@ -64,11 +87,30 @@ void ACameraCalibrationCheckerboard::OnConstruction(const FTransform& Transform)
 	Rebuild();
 }
 
+void ACameraCalibrationCheckerboard::ClearInstanceCheckeboardMeshComponents()
+{
+	// Need to cache because calling destroy will remove them from InstanceComponents
+	TArray<UActorComponent*> CachedComponents = GetInstanceComponents();
+
+	// Run in reverse to reduce memory churn when the components are removed from InstanceComponents
+	for (UActorComponent* CachedComponent : ReverseIterate(CachedComponents))
+	{
+		if (UStaticMeshComponent* CachedMeshComponent = Cast<UStaticMeshComponent>(CachedComponent))
+		{
+			if (CachedMeshComponent->GetStaticMesh() == CubeMesh)
+			{
+				CachedMeshComponent->DestroyComponent();
+			}
+		}
+	}
+}
+
 void ACameraCalibrationCheckerboard::Rebuild()
 {
 	Modify();
 
-	ClearInstanceComponents(true);
+	// Removes the instance components that make up the checkerboard
+	ClearInstanceCheckeboardMeshComponents();
 
 	if (!CubeMesh)
 	{
diff --git a/Engine/Plugins/VirtualProduction/CameraCalibrationCore/Source/CameraCalibrationCore/Public/CameraCalibrationCheckerboard.h b/Engine/Plugins/VirtualProduction/CameraCalibrationCore/Source/CameraCalibrationCore/Public/CameraCalibrationCheckerboard.h
--- a/Engine/Plugins/VirtualProduction/CameraCalibrationCore/Source/CameraCalibrationCore/Public/CameraCalibrationCheckerboard.h
+++ b/Engine/Plugins/VirtualProduction/CameraCalibrationCore/Source/CameraCalibrationCore/Public/CameraCalibrationCheckerboard.h
@@ -37,6 +37,10 @@ public:
 #endif
 	//~ End UObject Interface
 
+private:
+	/** Removes mesh components that make up the checkerboard */
+	void ClearInstanceCheckeboardMeshComponents();
+
 public:
 
 	/** Root component, gives the Actor a transform */
