Fix crash when using camera calibration subsystem in livelink hub

In the hub, the camera calibration subsystem isn't initialized when the camera calibration settings try to access it.

diff --git a/Engine/Plugins/VirtualProduction/CameraCalibrationCore/Source/CameraCalibrationCore/Private/CameraCalibrationSettings.cpp b/Engine/Plugins/VirtualProduction/CameraCalibrationCore/Source/CameraCalibrationCore/Private/CameraCalibrationSettings.cpp
--- a/Engine/Plugins/VirtualProduction/CameraCalibrationCore/Source/CameraCalibrationCore/Private/CameraCalibrationSettings.cpp
+++ b/Engine/Plugins/VirtualProduction/CameraCalibrationCore/Source/CameraCalibrationCore/Private/CameraCalibrationSettings.cpp
@@ -23,12 +23,14 @@ UCameraCalibrationSettings::UCameraCalibrationSettings()
 #if WITH_EDITORONLY_DATA
 	auto InitOverlayOverrides = [this]()
 	{
-		UCameraCalibrationSubsystem* SubSystem = GEngine->GetEngineSubsystem<UCameraCalibrationSubsystem>();
-		for (const FName& OverlayName : SubSystem->GetOverlayMaterialNames())
+		if (UCameraCalibrationSubsystem* SubSystem = GEngine->GetEngineSubsystem<UCameraCalibrationSubsystem>())
 		{
-			if (!CalibrationOverlayMaterialOverrides.Contains(OverlayName))
+			for (const FName& OverlayName : SubSystem->GetOverlayMaterialNames())
 			{
-				CalibrationOverlayMaterialOverrides.Add(OverlayName, nullptr);
+				if (!CalibrationOverlayMaterialOverrides.Contains(OverlayName))
+				{
+					CalibrationOverlayMaterialOverrides.Add(OverlayName, nullptr);
+				}
 			}
 		}
 	};
