Camera Calibration: Added the ability to modify the focus curves of a lens file, and changed the interpolation functions to account for the focus curves

diff --git a/Engine/Plugins/VirtualProduction/CameraCalibration/Source/CameraCalibrationEditor/Private/AssetEditor/Curves/LensMultiAxisCurveModel.cpp b/Engine/Plugins/VirtualProduction/CameraCalibration/Source/CameraCalibrationEditor/Private/AssetEditor/Curves/LensMultiAxisCurveModel.cpp
--- a/Engine/Plugins/VirtualProduction/CameraCalibration/Source/CameraCalibrationEditor/Private/AssetEditor/Curves/LensMultiAxisCurveModel.cpp
+++ b/Engine/Plugins/VirtualProduction/CameraCalibration/Source/CameraCalibrationEditor/Private/AssetEditor/Curves/LensMultiAxisCurveModel.cpp
@@ -50,12 +50,6 @@ void FLensDataMultiAxisCurveModel::SetKeyPositions(TArrayView<const FKeyHandle>
 
 void FLensDataMultiAxisCurveModel::SetKeyAttributes(TArrayView<const FKeyHandle> InKeys, TArrayView<const FKeyAttributes> InAttributes, EPropertyChangeType::Type ChangeType)
 {
-	// For now, don't allow focus axis curves to be editable
-	if (CurveAxis == ELensCurveAxis::Focus)
-	{
-		return;
-	}
-		
 	if (FBaseLensTable* DataTable = LensFile->GetDataTable(Category))
 	{
 		if (!DataTable->CanEditCurveKeyAttributes(ParameterIndex))
@@ -64,8 +58,15 @@ void FLensDataMultiAxisCurveModel::SetKeyAttributes(TArrayView<const FKeyHandle>
 		}
 			
 		FLensDataCurveModel::SetKeyAttributes(InKeys, InAttributes, ChangeType);
-		
-		DataTable->SetParameterCurveKeysAtFocus(CurveValue, ParameterIndex, CurrentCurve, InKeys);
+
+		if (CurveAxis == ELensCurveAxis::Zoom)
+		{
+			DataTable->SetParameterCurveKeysAtFocus(CurveValue, ParameterIndex, CurrentCurve, InKeys);
+		}
+		else if (CurveAxis == ELensCurveAxis::Focus)
+		{
+			DataTable->SetParameterCurveKeysAtZoom(CurveValue, ParameterIndex, CurrentCurve, InKeys);
+		}
 	}
 }
 
diff --git a/Engine/Plugins/VirtualProduction/CameraCalibrationCore/Shaders/Private/DisplacementMapBlending.usf b/Engine/Plugins/VirtualProduction/CameraCalibrationCore/Shaders/Private/DisplacementMapBlending.usf
--- a/Engine/Plugins/VirtualProduction/CameraCalibrationCore/Shaders/Private/DisplacementMapBlending.usf
+++ b/Engine/Plugins/VirtualProduction/CameraCalibrationCore/Shaders/Private/DisplacementMapBlending.usf
@@ -6,16 +6,13 @@
 
 #include "/Engine/Private/Common.ush"
 
-float EvalTime;
-float Curve0Key0Time;
-float Curve0Key1Time;
-float Curve0Key0Tangent;
-float Curve0Key1Tangent;
-float Curve1Key0Time;
-float Curve1Key1Time;
-float Curve1Key0Tangent;
-float Curve1Key1Tangent;
-float FocusBlendFactor;
+
+float EvalZoom;
+float EvalFocus;
+
+// Stores the x and y values and the x and y tangents for each corner of the blending patch
+// Ordered (X0, Y0) -> (X1, Y0) -> (X1, Y1) -> (X0, Y1)
+float4 PatchCorners[4];
 
 float2 FxFyScale;
 float2 PrincipalPoint;
@@ -63,6 +60,32 @@ float4 EvalAtTwoPoints(float EvalTime, float Time0, float Time1, float4 Value0,
 	return BezierInterp(P0, P1, P2, P3, Alpha);
 }
 
+float4 CoonsPatchBlend(float X, float Y, float4 Values[4])
+{
+	float Alpha = 0.0;
+	if (PatchCorners[0].x != PatchCorners[1].x)
+	{
+		Alpha = (X - PatchCorners[0].x) / (PatchCorners[1].x - PatchCorners[0].x);
+	}
+	
+	float Beta = 0.0;
+	if (PatchCorners[0].y != PatchCorners[3].y)
+	{
+		Beta = (Y - PatchCorners[0].y) / (PatchCorners[3].y - PatchCorners[0].y);
+	}
+	
+	float4 X0Curve = EvalAtTwoPoints(X, PatchCorners[0].x, PatchCorners[1].x, Values[0], Values[1], PatchCorners[0].z, PatchCorners[1].z);
+	float4 X1Curve = EvalAtTwoPoints(X, PatchCorners[3].x, PatchCorners[2].x, Values[3], Values[2], PatchCorners[3].z, PatchCorners[2].z);
+	float4 Y0Curve = EvalAtTwoPoints(Y, PatchCorners[0].y, PatchCorners[3].y, Values[0], Values[3], PatchCorners[0].w, PatchCorners[3].w);
+	float4 Y1Curve = EvalAtTwoPoints(Y, PatchCorners[1].y, PatchCorners[2].y, Values[1], Values[2], PatchCorners[1].w, PatchCorners[2].w);
+	
+	float4 Lx = lerp(X0Curve, X1Curve, Beta);
+	float4 Ly = lerp(Y0Curve, Y1Curve, Alpha);
+	float4 B = lerp(lerp(Values[0], Values[1], Alpha), lerp(Values[3], Values[2], Alpha), Beta);
+
+	return Lx + Ly - B;
+}
+
 void BlendPS(float4 SvPosition : SV_POSITION, out float4 OutColor : SV_Target0)
 {
 	float2 UV = SvPosition.xy / OutputTextureExtent.xy;
@@ -71,31 +94,31 @@ void BlendPS(float4 SvPosition : SV_POSITION, out float4 OutColor : SV_Target0)
 	float4 TexOneColor = Texture2DSampleLevel(SourceTextureOne, SourceTextureSampler, UV, 0);
 	float4 BlendedColor;
 
-// Case - 2 Focus Curves, 1 Zoom Point Each - Just Lerp
+// Case - 2 Focus Curves, 1 Zoom Point Each
 #if BLEND_TYPE == 1
 	
 	float4 TexTwoColor = Texture2DSampleLevel(SourceTextureTwo, SourceTextureSampler, UV, 0);
-	BlendedColor = lerp(TexOneColor, TexTwoColor, FocusBlendFactor);
-
-// Case - 1 Focus Curve, 2 Zoom Points - One Bezier Interp
+	BlendedColor = EvalAtTwoPoints(EvalFocus, PatchCorners[0].y, PatchCorners[1].y, TexOneColor, TexTwoColor, PatchCorners[0].w, PatchCorners[1].w);
+	
+// Case - 1 Focus Curve, 2 Zoom Points
 #elif BLEND_TYPE == 2
 
 	float4 TexTwoColor = Texture2DSampleLevel(SourceTextureTwo, SourceTextureSampler, UV, 0);
-	BlendedColor = EvalAtTwoPoints(EvalTime, Curve0Key0Time, Curve0Key1Time, TexOneColor, TexTwoColor, Curve0Key0Tangent, Curve0Key1Tangent);
+	BlendedColor = EvalAtTwoPoints(EvalZoom, PatchCorners[0].x, PatchCorners[1].x, TexOneColor, TexTwoColor, PatchCorners[0].z, PatchCorners[1].z);
 	
-// Case - 2 Focus Curves, 2 Zoom Points Each - Two Bezier Interps, then Lerp
+// Case - 2 Focus Curves, 2 Zoom Points Each
 #elif BLEND_TYPE == 3
 
-	float4 TexTwoColor = Texture2DSampleLevel(SourceTextureTwo, SourceTextureSampler, UV, 0);
-	float4 TexThreeColor = Texture2DSampleLevel(SourceTextureThree, SourceTextureSampler, UV, 0);
-	float4 TexFourColor = Texture2DSampleLevel(SourceTextureFour, SourceTextureSampler, UV, 0);
-
-	float4 Curve0BlendedColor = EvalAtTwoPoints(EvalTime, Curve0Key0Time, Curve0Key1Time, TexOneColor, TexTwoColor, Curve0Key0Tangent, Curve0Key1Tangent);
-	float4 Curve1BlendedColor = EvalAtTwoPoints(EvalTime, Curve1Key0Time, Curve1Key1Time, TexThreeColor, TexFourColor, Curve1Key0Tangent, Curve1Key1Tangent);
-
-	//Lerp across focus curves
-	BlendedColor = lerp(Curve0BlendedColor, Curve1BlendedColor, FocusBlendFactor);
-
+	float4 TexColors[4] =
+	{
+		TexOneColor,
+		Texture2DSampleLevel(SourceTextureTwo, SourceTextureSampler, UV, 0),
+		Texture2DSampleLevel(SourceTextureThree, SourceTextureSampler, UV, 0),
+		Texture2DSampleLevel(SourceTextureFour, SourceTextureSampler, UV, 0)
+	};
+	
+	BlendedColor = CoonsPatchBlend(EvalZoom, EvalFocus, TexColors);
+	
 #else
 	
 	BlendedColor = TexOneColor;
diff --git a/Engine/Plugins/VirtualProduction/CameraCalibrationCore/Source/CameraCalibrationCore/Private/LensFile.cpp b/Engine/Plugins/VirtualProduction/CameraCalibrationCore/Source/CameraCalibrationCore/Private/LensFile.cpp
--- a/Engine/Plugins/VirtualProduction/CameraCalibrationCore/Source/CameraCalibrationCore/Private/LensFile.cpp
+++ b/Engine/Plugins/VirtualProduction/CameraCalibrationCore/Source/CameraCalibrationCore/Private/LensFile.cpp
@@ -12,6 +12,7 @@
 #include "Engine/Engine.h"
 #include "Engine/TextureRenderTarget2D.h"
 #include "LensDistortionModelHandlerBase.h"
+#include "LensFileObjectVersion.h"
 #include "LensFileRendering.h"
 #include "LensInterpolationUtils.h"
 #include "Curves/CurveEvaluation.h"
@@ -192,6 +193,23 @@ ULensFile::ULensFile()
 	}
 }
 
+void ULensFile::Serialize(FArchive& Ar)
+{
+	Super::Serialize(Ar);
+
+	Ar.UsingCustomVersion(FLensFileObjectVersion::GUID);
+
+#if WITH_EDITOR
+	if (Ar.IsLoading())
+	{
+		if (Ar.CustomVer(FLensFileObjectVersion::GUID) < FLensFileObjectVersion::EditableFocusCurves)
+		{
+			BuildLensTableFocusCurves();
+		}
+	}
+#endif 
+}
+
 #if WITH_EDITOR
 
 void ULensFile::PostEditChangeChainProperty(struct FPropertyChangedChainEvent& PropertyChangedEvent)
@@ -240,74 +258,33 @@ bool ULensFile::EvaluateDistortionParameters(float InFocus, float InZoom, FDisto
 {
 	TRACE_CPUPROFILER_EVENT_SCOPE(ULensFile::EvaluateDistortionParameters);
 
-	if (DistortionTable.GetFocusPoints().Num() <= 0)
-	{
-		return false;
-	}
-
-	const LensDataTableUtils::FPointNeighbors Neighbors = LensDataTableUtils::FindFocusPoints(InFocus, DistortionTable.GetFocusPoints());
-	const FDistortionFocusPoint& PreviousPoint = DistortionTable.GetFocusPoints()[Neighbors.PreviousIndex];
-	const FDistortionFocusPoint& NextPoint = DistortionTable.GetFocusPoints()[Neighbors.NextIndex];
-	if (Neighbors.NextIndex == Neighbors.PreviousIndex)
-	{
-		const LensDataTableUtils::FPointNeighbors ZoomNeighbors = LensDataTableUtils::FindZoomPoints(InZoom, PreviousPoint.ZoomPoints);
-		const FDistortionZoomPoint& PreviousZoomPoint = PreviousPoint.ZoomPoints[ZoomNeighbors.PreviousIndex];
-		const FDistortionZoomPoint& NextZoomPoint = PreviousPoint.ZoomPoints[ZoomNeighbors.NextIndex];
-		if (ZoomNeighbors.NextIndex == ZoomNeighbors.PreviousIndex)
-		{
-			//Exactly on one point
-			OutEvaluatedValue = PreviousZoomPoint.DistortionInfo;
-		}
-		else
+	using FBlendParams = LensInterpolationUtils::FDistortionMapBlendParams<FDistortionTable>;
+	using FBlendResults = LensInterpolationUtils::FDistortionMapBlendResults;
+	
+	FBlendParams Params;
+	
+	Params.GetDistortionParameters = FBlendParams::FGetDistortionParameters::CreateLambda(
+		[](const FDistortionFocusPoint& FocusPoint, const FDistortionFocusCurve& FocusCurve)
 		{
-			//Blend parameters following map blending curve 
-			const FRichCurveKey Curve0Key0 = PreviousPoint.MapBlendingCurve.Keys[ZoomNeighbors.PreviousIndex];
-			const FRichCurveKey Curve0Key1 = PreviousPoint.MapBlendingCurve.Keys[ZoomNeighbors.NextIndex];
-			const float Times[2] = { Curve0Key0.Time, Curve0Key1.Time };
-			const float Tangents[2] = { Curve0Key0.LeaveTangent, Curve0Key1.ArriveTangent };
-			const TConstArrayView<float> Parameters[2] = { PreviousZoomPoint.DistortionInfo.Parameters, NextZoomPoint.DistortionInfo.Parameters };
-			LensFileUtils::FindWeightsAndInterp(InZoom, Times, Tangents, TOptional<float>(), Parameters, OutEvaluatedValue.Parameters);
+			FDistortionInfo Point;
+			if (FocusPoint.GetPoint(FocusCurve.Zoom, Point))
+			{
+				return TOptional(Point);
+			}
+			
+			return TOptional<FDistortionInfo>();
 		}
-	}
-	else
+	);
+	
+	Params.DistortionParamNum = OutEvaluatedValue.Parameters.Num();
+
+	FBlendResults Results = LensInterpolationUtils::DistortionMapBlend(DistortionTable, InFocus, InZoom, Params);
+	if (!Results.bValid)
 	{
-		//Previous Focus two zoom points
-		const LensDataTableUtils::FPointNeighbors PreviousFocusZoomNeighbors = LensDataTableUtils::FindZoomPoints(InZoom, PreviousPoint.ZoomPoints);
-		const FDistortionZoomPoint& PreviousFocusPreviousZoomPoint = PreviousPoint.ZoomPoints[PreviousFocusZoomNeighbors.PreviousIndex];
-		const FDistortionZoomPoint& PreviousFocusNextZoomPoint = PreviousPoint.ZoomPoints[PreviousFocusZoomNeighbors.NextIndex];
-
-		//Next focus two zoom points
-		const LensDataTableUtils::FPointNeighbors NextFocusZoomNeighbors = LensDataTableUtils::FindZoomPoints(InZoom, NextPoint.ZoomPoints);
-		const FDistortionZoomPoint NextFocusPreviousZoomPoint = NextPoint.ZoomPoints[NextFocusZoomNeighbors.PreviousIndex];
-		const FDistortionZoomPoint& NextFocusNextZoomPoint = NextPoint.ZoomPoints[NextFocusZoomNeighbors.NextIndex];
-
-		//Verify if we are dealing with one zoom point on each focus. If that's the case, we are doing simple lerp across the focus curves
-		if (PreviousFocusZoomNeighbors.NextIndex == PreviousFocusZoomNeighbors.PreviousIndex
-			&& NextFocusZoomNeighbors.NextIndex == NextFocusZoomNeighbors.PreviousIndex)
-		{
-			//Linear blend between each zoom point pair and then both results linearly according to focus
-			const float FocusBlendFactor = LensInterpolationUtils::GetBlendFactor(InFocus, PreviousPoint.Focus, NextPoint.Focus);
-			LensInterpolationUtils::Interpolate(FocusBlendFactor, &PreviousFocusPreviousZoomPoint.DistortionInfo, &NextFocusPreviousZoomPoint.DistortionInfo, &OutEvaluatedValue);
-		}
-		else
-		{
-			//Blend parameters following both zoom curves and lerp result using focus blend factor
-			const float FocusBlendFactor = LensInterpolationUtils::GetBlendFactor(InFocus, PreviousPoint.Focus, NextPoint.Focus);
-			const FRichCurveKey Curve0Key0 = PreviousPoint.MapBlendingCurve.Keys[PreviousFocusZoomNeighbors.PreviousIndex];
-			const FRichCurveKey Curve0Key1 = PreviousPoint.MapBlendingCurve.Keys[PreviousFocusZoomNeighbors.NextIndex];
-			const FRichCurveKey Curve1Key0 = NextPoint.MapBlendingCurve.Keys[NextFocusZoomNeighbors.PreviousIndex];
-			const FRichCurveKey Curve1Key1 = NextPoint.MapBlendingCurve.Keys[NextFocusZoomNeighbors.NextIndex];
-
-			const float Times[4] = { Curve0Key0.Time, Curve0Key1.Time, Curve1Key0.Time, Curve1Key1.Time };
-			const float Tangents[4] = { Curve0Key0.LeaveTangent, Curve0Key1.ArriveTangent, Curve1Key0.LeaveTangent, Curve1Key1.ArriveTangent };
-			const TConstArrayView<float> Parameters[4] = { PreviousFocusPreviousZoomPoint.DistortionInfo.Parameters
-														 , PreviousFocusNextZoomPoint.DistortionInfo.Parameters
-														 , NextFocusPreviousZoomPoint.DistortionInfo.Parameters
-														 , NextFocusNextZoomPoint.DistortionInfo.Parameters };
-			LensFileUtils::FindWeightsAndInterp(InZoom, Times, Tangents, FocusBlendFactor, Parameters, OutEvaluatedValue.Parameters);
-		}
+		return false;
 	}
 
+	OutEvaluatedValue = Results.BlendedDistortionParams.GetValue();
 	return true;
 }
 
@@ -315,78 +292,36 @@ bool ULensFile::EvaluateFocalLength(float InFocus, float InZoom, FFocalLengthInf
 {
 	TRACE_CPUPROFILER_EVENT_SCOPE(ULensFile::EvaluateFocalLength);
 
-	if (FocalLengthTable.FocusPoints.Num() <= 0)
+	constexpr int32 NumParams = 2;
+	TArray<float> BlendedParameters;
+	if (LensInterpolationUtils::IndexedParameterBlend(FocalLengthTable.FocusPoints, FocalLengthTable.FocusCurves, InFocus, InZoom, NumParams, BlendedParameters))
 	{
-		return false;
-	}
+		ensure(BlendedParameters.Num() == NumParams);
+		OutEvaluatedValue.FxFy.X = BlendedParameters[FFocalLengthTable::FParameters::Fx];
+		OutEvaluatedValue.FxFy.Y = BlendedParameters[FFocalLengthTable::FParameters::Fy];
 
-	if (FocalLengthTable.FocusPoints.Num() == 1)
-	{
-		OutEvaluatedValue.FxFy.X = FocalLengthTable.FocusPoints[0].Fx.Eval(InZoom);
-		OutEvaluatedValue.FxFy.Y = FocalLengthTable.FocusPoints[0].Fy.Eval(InZoom);
 		return true;
 	}
-
-	const LensDataTableUtils::FPointNeighbors Neighbors = LensDataTableUtils::FindFocusPoints(InFocus, FocalLengthTable.GetFocusPoints());
-	if (Neighbors.NextIndex == Neighbors.PreviousIndex)
-	{
-		OutEvaluatedValue.FxFy.X = FocalLengthTable.FocusPoints[Neighbors.PreviousIndex].Fx.Eval(InZoom);
-		OutEvaluatedValue.FxFy.Y = FocalLengthTable.FocusPoints[Neighbors.PreviousIndex].Fy.Eval(InZoom);
-	}
-	else
-	{
-		FVector2D PreviousValue;
-		PreviousValue.X = FocalLengthTable.FocusPoints[Neighbors.PreviousIndex].Fx.Eval(InZoom);
-		PreviousValue.Y = FocalLengthTable.FocusPoints[Neighbors.PreviousIndex].Fy.Eval(InZoom);
-		FVector2D NextValue;
-		NextValue.X = FocalLengthTable.FocusPoints[Neighbors.NextIndex].Fx.Eval(InZoom);
-		NextValue.Y = FocalLengthTable.FocusPoints[Neighbors.NextIndex].Fy.Eval(InZoom);
-
-		//Blend result between focus
-		const float BlendFactor = LensInterpolationUtils::GetBlendFactor(InFocus, FocalLengthTable.FocusPoints[Neighbors.PreviousIndex].Focus, FocalLengthTable.FocusPoints[Neighbors.NextIndex].Focus);
-		OutEvaluatedValue.FxFy = LensInterpolationUtils::BlendValue(BlendFactor, PreviousValue, NextValue);
-	}
-
-	return true;
+	
+	return false;
 }
 
 bool ULensFile::EvaluateImageCenterParameters(float InFocus, float InZoom, FImageCenterInfo& OutEvaluatedValue) const
 {
 	TRACE_CPUPROFILER_EVENT_SCOPE(ULensFile::EvaluateImageCenterParameters);
 
-	if (ImageCenterTable.FocusPoints.Num() <= 0)
+	constexpr int32 NumParams = 2;
+	TArray<float> BlendedParameters;
+	if (LensInterpolationUtils::IndexedParameterBlend(ImageCenterTable.FocusPoints, ImageCenterTable.FocusCurves, InFocus, InZoom, NumParams, BlendedParameters))
 	{
-		return false;
-	}
+		ensure(BlendedParameters.Num() == NumParams);
+		OutEvaluatedValue.PrincipalPoint.X = BlendedParameters[FImageCenterTable::FParameters::Cx];
+		OutEvaluatedValue.PrincipalPoint.Y = BlendedParameters[FImageCenterTable::FParameters::Cy];
 
-	if (ImageCenterTable.FocusPoints.Num() == 1)
-	{
-		OutEvaluatedValue.PrincipalPoint.X = ImageCenterTable.FocusPoints[0].Cx.Eval(InZoom);
-		OutEvaluatedValue.PrincipalPoint.Y = ImageCenterTable.FocusPoints[0].Cy.Eval(InZoom);
 		return true;
 	}
-
-	const LensDataTableUtils::FPointNeighbors Neighbors = LensDataTableUtils::FindFocusPoints(InFocus, ImageCenterTable.GetFocusPoints());
-	if (Neighbors.NextIndex == Neighbors.PreviousIndex)
-	{
-		OutEvaluatedValue.PrincipalPoint.X = ImageCenterTable.FocusPoints[Neighbors.PreviousIndex].Cx.Eval(InZoom);
-		OutEvaluatedValue.PrincipalPoint.Y = ImageCenterTable.FocusPoints[Neighbors.PreviousIndex].Cy.Eval(InZoom);
-	}
-	else
-	{
-		FVector2D PreviousValue;
-		PreviousValue.X = ImageCenterTable.FocusPoints[Neighbors.PreviousIndex].Cx.Eval(InZoom);
-		PreviousValue.Y = ImageCenterTable.FocusPoints[Neighbors.PreviousIndex].Cy.Eval(InZoom);
-		FVector2D NextValue;
-		NextValue.X = ImageCenterTable.FocusPoints[Neighbors.NextIndex].Cx.Eval(InZoom);
-		NextValue.Y = ImageCenterTable.FocusPoints[Neighbors.NextIndex].Cy.Eval(InZoom);
-
-		//Blend result between focus
-		const float BlendFactor = LensInterpolationUtils::GetBlendFactor(InFocus, ImageCenterTable.FocusPoints[Neighbors.PreviousIndex].Focus, ImageCenterTable.FocusPoints[Neighbors.NextIndex].Focus);
-		OutEvaluatedValue.PrincipalPoint = LensInterpolationUtils::BlendValue(BlendFactor, PreviousValue, NextValue);
-	}
-
-	return true;
+	
+	return false;
 }
 
 bool ULensFile::EvaluateDistortionData(float InFocus, float InZoom, FVector2D InFilmback, ULensDistortionModelHandlerBase* InLensHandler) const
@@ -472,236 +407,111 @@ void ULensFile::SetupNoDistortionOutput(ULensDistortionModelHandlerBase* LensHan
 
 bool ULensFile::EvaluateDistortionForParameters(float InFocus, float InZoom, FVector2D InFilmback, ULensDistortionModelHandlerBase* InLensHandler) const
 {
-	//ImageCenter is always required
-	FImageCenterInfo ImageCenter;
-	EvaluateImageCenterParameters(InFocus, InZoom, ImageCenter);
+	// Compute interpolated image center and focal length to pass to the lens handler
+	FImageCenterInfo InterpolatedImageCenter;
+	EvaluateImageCenterParameters(InFocus, InZoom, InterpolatedImageCenter);
 
+	FFocalLengthInfo InterpolatedFocalLength;
+	EvaluateFocalLength(InFocus, InZoom, InterpolatedFocalLength);
+	
 	FCameraFilmbackSettings CameraFilmback;
 	CameraFilmback.SensorWidth = InFilmback.X;
 	CameraFilmback.SensorHeight = InFilmback.Y;
 
-	//No distortion parameters case. Still process to have center shift
-	if (DistortionTable.GetFocusPoints().Num() <= 0)
-	{
-		// Initialize all distortion parameters with their default values
-		FDistortionInfo DistortionPoint;
-		LensInfo.LensModel->GetDefaultObject<ULensModel>()->GetDefaultParameterArray(DistortionPoint.Parameters);
-
-		//Setup handler state based on evaluated parameters. If none were found, no distortion will be returned
-		FLensDistortionState State;
-		State.DistortionInfo.Parameters = MoveTemp(DistortionPoint.Parameters);
-
-		//Evaluate Focal Length
-		FFocalLengthInfo FocalLength;
-		EvaluateFocalLength(InFocus, InZoom, FocalLength);
-		const FVector2D FxFyScale = FVector2D(LensInfo.SensorDimensions.X / InFilmback.X, LensInfo.SensorDimensions.Y / InFilmback.Y);
-		State.FocalLengthInfo.FxFy = FocalLength.FxFy * FxFyScale;
-		State.ImageCenter.PrincipalPoint = ImageCenter.PrincipalPoint;
-
-		//Updates handler state
-		InLensHandler->SetDistortionState(State);
-		InLensHandler->SetCameraFilmback(CameraFilmback);
-
-		InLensHandler->SetOverscanFactor(1.0f);
-
-		//Draw displacement map associated with the new state
-		InLensHandler->ProcessCurrentDistortion();
-	}
-	else
-	{
-		FLensDistortionState State;
-		State.ImageCenter.PrincipalPoint = ImageCenter.PrincipalPoint;
-
-		FDisplacementMapBlendingParams Params;
-		float InterpolatedOverscanFactor = 1.0f;
+	const FVector2D FxFyScale = FVector2D(LensInfo.SensorDimensions.X / CameraFilmback.SensorWidth, LensInfo.SensorDimensions.Y / CameraFilmback.SensorHeight);
+	
+	FLensDistortionState InterpolatedState;
+	InterpolatedState.FocalLengthInfo.FxFy = InterpolatedFocalLength.FxFy * FxFyScale;
+	InterpolatedState.ImageCenter.PrincipalPoint = InterpolatedImageCenter.PrincipalPoint;
+		
+	// Initialize all distortion parameters with their default values
+	LensInfo.LensModel->GetDefaultObject<ULensModel>()->GetDefaultParameterArray(InterpolatedState.DistortionInfo.Parameters);
+	
+	LensInterpolationUtils::FDistortionMapBlendParams<FDistortionTable> Params;
+	Params.bGenerateBlendingParams = true;
+	Params.UndistortedMaps = UndistortionDisplacementMapHolders;
+	Params.DistortedMaps = DistortionDisplacementMapHolders;
+	Params.DistortionParamNum = InterpolatedState.DistortionInfo.Parameters.Num();
+
+	// Callback that retrieves the distortion parameters for the specified focus and zoom
+	Params.GetDistortionParameters = LensInterpolationUtils::FDistortionMapBlendParams<FDistortionTable>::FGetDistortionParameters::CreateLambda(
+		[](const FDistortionFocusPoint& FocusPoint, const FDistortionFocusCurve& FocusCurve)
+		{
+			FDistortionInfo Point;
+			if (FocusPoint.GetPoint(FocusCurve.Zoom, Point))
+			{
+				return TOptional(Point);
+			}
+			
+			return TOptional<FDistortionInfo>();
+		}
+	);
 
-		//Helper function to compute the current distortion state
-		const auto ProcessDistortionState = [this, &State, CameraFilmback, InLensHandler](const FDistortionInfo& DistortionInfo, const FFocalLengthInfo& FocalLength, UTextureRenderTarget2D* UndistortionRenderTarget, UTextureRenderTarget2D* DistortionRenderTarget, float& OverscanFactor)
+	// Callback when the blend function constructs the displacement maps for each corner used in the blend, which generates the displacement maps
+	// and returns the computed overscan of the displacement map
+	Params.ProcessDisplacementMaps = LensInterpolationUtils::FDistortionMapBlendParams<FDistortionTable>::FProcessDisplacementMaps::CreateLambda(
+		[this, InLensHandler, &CameraFilmback, &InterpolatedImageCenter, &FxFyScale]
+		(const FDistortionFocusPoint& FocusPoint, const FDistortionFocusCurve& FocusCurve, UTextureRenderTarget2D* InUndistortedMap, UTextureRenderTarget2D* InDistortedMap)
 		{
-			State.DistortionInfo.Parameters = DistortionInfo.Parameters;
+			FLensDistortionState State;
+			State.ImageCenter.PrincipalPoint = InterpolatedImageCenter.PrincipalPoint;
+
+			// In case the point doesn't exist, we need to fill the distortion parameter array with default values
+			LensInfo.LensModel->GetDefaultObject<ULensModel>()->GetDefaultParameterArray(State.DistortionInfo.Parameters);
+			FocusPoint.GetPoint(FocusCurve.Zoom, State.DistortionInfo);
 
-			const FVector2D FxFyScale = FVector2D(LensInfo.SensorDimensions.X / CameraFilmback.SensorWidth, LensInfo.SensorDimensions.Y / CameraFilmback.SensorHeight);
+			FFocalLengthInfo FocalLength;
+			LensDataTableUtils::GetPointValue<FFocalLengthFocusPoint>(FocusPoint.Focus, FocusCurve.Zoom, FocalLengthTable.FocusPoints, FocalLength);
+			
 			State.FocalLengthInfo.FxFy = FocalLength.FxFy * FxFyScale;
 
 			InLensHandler->SetDistortionState(State);
 			InLensHandler->SetCameraFilmback(CameraFilmback);
-			InLensHandler->DrawUndistortionDisplacementMap(UndistortionRenderTarget);
-			InLensHandler->DrawDistortionDisplacementMap(DistortionRenderTarget);
-
-			OverscanFactor = InLensHandler->ComputeOverscanFactor();
-		};
-
-		//Find focuses in play
-		const LensDataTableUtils::FPointNeighbors Neighbors = LensDataTableUtils::FindFocusPoints(InFocus, DistortionTable.GetFocusPoints());
-		const FDistortionFocusPoint& PreviousPoint = DistortionTable.GetFocusPoints()[Neighbors.PreviousIndex];
-		const FDistortionFocusPoint& NextPoint = DistortionTable.GetFocusPoints()[Neighbors.NextIndex];
-		if (Neighbors.NextIndex == Neighbors.PreviousIndex)
-		{
-			const LensDataTableUtils::FPointNeighbors ZoomNeighbors = LensDataTableUtils::FindZoomPoints(InZoom, PreviousPoint.ZoomPoints);
-			const FDistortionZoomPoint& PreviousZoomPoint = PreviousPoint.ZoomPoints[ZoomNeighbors.PreviousIndex];
-			const FDistortionZoomPoint& NextZoomPoint = NextPoint.ZoomPoints[ZoomNeighbors.NextIndex];
+			InLensHandler->DrawUndistortionDisplacementMap(InUndistortedMap);
+			InLensHandler->DrawDistortionDisplacementMap(InDistortedMap);
 
-			//Get FocalLength points
-			FFocalLengthInfo PreviousZoomPointFocalLength;
-			FFocalLengthInfo NextZoomPointFocalLength;
-			LensDataTableUtils::GetPointValue(PreviousPoint.Focus, PreviousZoomPoint.Zoom, FocalLengthTable.GetFocusPoints(), PreviousZoomPointFocalLength);
-			LensDataTableUtils::GetPointValue(PreviousPoint.Focus, NextZoomPoint.Zoom, FocalLengthTable.GetFocusPoints(), NextZoomPointFocalLength);
-
-			if (ZoomNeighbors.NextIndex == ZoomNeighbors.PreviousIndex)
-			{
-				//Exactly on one point
-				Params.BlendType = EDisplacementMapBlendType::OneFocusOneZoom;
-
-				ProcessDistortionState(PreviousZoomPoint.DistortionInfo, PreviousZoomPointFocalLength, UndistortionDisplacementMapHolders[0], DistortionDisplacementMapHolders[0], InterpolatedOverscanFactor);
-			}
-			else
-			{
-				//Interpolate between two zoom points following the map blending curve
-				Params.BlendType = EDisplacementMapBlendType::OneFocusTwoZoom;
-
-				float BlendedOverscanFactors[2];
-				ProcessDistortionState(PreviousZoomPoint.DistortionInfo, PreviousZoomPointFocalLength, UndistortionDisplacementMapHolders[0], DistortionDisplacementMapHolders[0], BlendedOverscanFactors[0]);
-				ProcessDistortionState(NextZoomPoint.DistortionInfo, NextZoomPointFocalLength, UndistortionDisplacementMapHolders[1], DistortionDisplacementMapHolders[1], BlendedOverscanFactors[1]);
-
-				//Set displacement map blending parameters
-				const TArray<FRichCurveKey>& PreviousPointsKeys = PreviousPoint.MapBlendingCurve.GetConstRefOfKeys();
-				const FRichCurveKey Curve0Key0 = PreviousPointsKeys[ZoomNeighbors.PreviousIndex];
-				const FRichCurveKey Curve0Key1 = PreviousPointsKeys[ZoomNeighbors.NextIndex];
-
-				Params.EvalTime = InZoom;
-				Params.Curve0Key0Time = Curve0Key0.Time;
-				Params.Curve0Key1Time = Curve0Key1.Time;
-				Params.Curve0Key0Tangent = Curve0Key0.LeaveTangent;
-				Params.Curve0Key1Tangent = Curve0Key1.ArriveTangent;
-
-				//Interpolate distortion parameters along the map blending curve
-				const float Times[2] = { Curve0Key0.Time, Curve0Key1.Time };
-				const float Tangents[2] = { Curve0Key0.LeaveTangent, Curve0Key1.ArriveTangent };
-				const TConstArrayView<float> Parameters[2] = { PreviousZoomPoint.DistortionInfo.Parameters, NextZoomPoint.DistortionInfo.Parameters };
-				LensFileUtils::FindWeightsAndInterp(InZoom, Times, Tangents, TOptional<float>(), Parameters, State.DistortionInfo.Parameters);
-
-				//Interpolate overscan factor along the map blending curve
-				const float OverscanFactors[2] = { BlendedOverscanFactors[0], BlendedOverscanFactors[1] };
-				LensFileUtils::FindWeightsAndInterp(InZoom, Times, Tangents, TOptional<float>(), OverscanFactors, InterpolatedOverscanFactor);
-			}
+			return InLensHandler->ComputeOverscanFactor();
 		}
-		else
-		{
-			//Previous Focus two zoom points
-			const LensDataTableUtils::FPointNeighbors PreviousFocusZoomNeighbors = LensDataTableUtils::FindZoomPoints(InZoom, PreviousPoint.ZoomPoints);
-			const FDistortionZoomPoint& PreviousFocusPreviousZoomPoint = PreviousPoint.ZoomPoints[PreviousFocusZoomNeighbors.PreviousIndex];
-			const FDistortionZoomPoint& PreviousFocusNextZoomPoint = PreviousPoint.ZoomPoints[PreviousFocusZoomNeighbors.NextIndex];
-			FFocalLengthInfo PreviousFocusPreviousZoomPointFocalLength;
-			FFocalLengthInfo PreviousFocusNextZoomPointFocalLength;
-			LensDataTableUtils::GetPointValue<FFocalLengthFocusPoint>(PreviousPoint.Focus, PreviousFocusPreviousZoomPoint.Zoom, FocalLengthTable.FocusPoints, PreviousFocusPreviousZoomPointFocalLength);
-			LensDataTableUtils::GetPointValue<FFocalLengthFocusPoint>(PreviousPoint.Focus, PreviousFocusNextZoomPoint.Zoom, FocalLengthTable.FocusPoints, PreviousFocusNextZoomPointFocalLength);
-
-			//Next focus two zoom points
-			const LensDataTableUtils::FPointNeighbors NextFocusZoomNeighbors = LensDataTableUtils::FindZoomPoints(InZoom, NextPoint.ZoomPoints);
-			const FDistortionZoomPoint& NextFocusPreviousZoomPoint = NextPoint.ZoomPoints[NextFocusZoomNeighbors.PreviousIndex];
-			const FDistortionZoomPoint& NextFocusNextZoomPoint = NextPoint.ZoomPoints[NextFocusZoomNeighbors.NextIndex];
-			FFocalLengthInfo NextFocusPreviousZoomPointFocalLength;
-			FFocalLengthInfo NextFocusNextZoomPointFocalLength;
-			LensDataTableUtils::GetPointValue<FFocalLengthFocusPoint>(NextPoint.Focus, NextFocusPreviousZoomPoint.Zoom, FocalLengthTable.FocusPoints, NextFocusPreviousZoomPointFocalLength);
-			LensDataTableUtils::GetPointValue<FFocalLengthFocusPoint>(NextPoint.Focus, NextFocusNextZoomPoint.Zoom, FocalLengthTable.FocusPoints, NextFocusNextZoomPointFocalLength);
-
-			const float FocusBlendFactor = LensInterpolationUtils::GetBlendFactor(InFocus, PreviousPoint.Focus, NextPoint.Focus);
-			Params.FocusBlendFactor = FocusBlendFactor;
-
-			//Verify if we are dealing with one zoom point on each focus. If that's the case, we are doing simple lerp across the focus curves
-			if(PreviousFocusZoomNeighbors.NextIndex == PreviousFocusZoomNeighbors.PreviousIndex
-				&& NextFocusZoomNeighbors.NextIndex == NextFocusZoomNeighbors.PreviousIndex)
-			{
-				//Linearly interpolate between two focus curves
-				Params.BlendType = EDisplacementMapBlendType::TwoFocusOneZoom;
+	);
+	
+	LensInterpolationUtils::FDistortionMapBlendResults Results = LensInterpolationUtils::DistortionMapBlend(DistortionTable, InFocus, InZoom, Params);
+	if (!Results.bValid)
+	{
+		// No distortion parameters case. Still process to have center shift
+		// Setup handler state based on evaluated parameters. If none were found, no distortion will be returned
+		InLensHandler->SetDistortionState(InterpolatedState);
+		InLensHandler->SetCameraFilmback(CameraFilmback);
 
-				float BlendedOverscanFactors[2];
-				ProcessDistortionState(PreviousFocusPreviousZoomPoint.DistortionInfo, PreviousFocusPreviousZoomPointFocalLength, UndistortionDisplacementMapHolders[0], DistortionDisplacementMapHolders[0], BlendedOverscanFactors[0]);
-				ProcessDistortionState(NextFocusPreviousZoomPoint.DistortionInfo, NextFocusPreviousZoomPointFocalLength, UndistortionDisplacementMapHolders[1], DistortionDisplacementMapHolders[1], BlendedOverscanFactors[1]);
+		InLensHandler->SetOverscanFactor(1.0f);
 
-				//Linearly interpolate the distortion parameters
-				LensInterpolationUtils::Interpolate(FocusBlendFactor, &PreviousFocusPreviousZoomPoint.DistortionInfo, &NextFocusPreviousZoomPoint.DistortionInfo, &State.DistortionInfo);
+		//Draw displacement map associated with the new state
+		InLensHandler->ProcessCurrentDistortion();
+		return true;
+	}
 
-				//Linearly interpolate the overscan factor
-				InterpolatedOverscanFactor = FMath::Lerp(BlendedOverscanFactors[0], BlendedOverscanFactors[1], FocusBlendFactor);
-			}
-			else
-			{
-				//Interpolate between two zoom points on each curve using the correct map blending curve, then linearly interpolate between the two focus points
-				Params.BlendType = EDisplacementMapBlendType::TwoFocusTwoZoom;
-
-				float BlendedOverscanFactors[4];
-				ProcessDistortionState(PreviousFocusPreviousZoomPoint.DistortionInfo, PreviousFocusPreviousZoomPointFocalLength, UndistortionDisplacementMapHolders[0], DistortionDisplacementMapHolders[0], BlendedOverscanFactors[0]);
-				ProcessDistortionState(PreviousFocusNextZoomPoint.DistortionInfo, PreviousFocusNextZoomPointFocalLength, UndistortionDisplacementMapHolders[1], DistortionDisplacementMapHolders[1], BlendedOverscanFactors[1]);
-				ProcessDistortionState(NextFocusPreviousZoomPoint.DistortionInfo, NextFocusPreviousZoomPointFocalLength, UndistortionDisplacementMapHolders[2], DistortionDisplacementMapHolders[2], BlendedOverscanFactors[2]);
-				ProcessDistortionState(NextFocusNextZoomPoint.DistortionInfo, NextFocusNextZoomPointFocalLength, UndistortionDisplacementMapHolders[3], DistortionDisplacementMapHolders[3], BlendedOverscanFactors[3]);
-
-				//Set displacement map blending parameters
-				const TArray<FRichCurveKey>& Curve0Keys = PreviousPoint.MapBlendingCurve.GetConstRefOfKeys();
-				const FRichCurveKey Curve0Key0 = Curve0Keys[PreviousFocusZoomNeighbors.PreviousIndex];
-				const FRichCurveKey Curve0Key1 = Curve0Keys[PreviousFocusZoomNeighbors.NextIndex];
-
-				const TArray<FRichCurveKey>& Curve1Keys = NextPoint.MapBlendingCurve.GetConstRefOfKeys();
-				const FRichCurveKey Curve1Key0 = Curve1Keys[NextFocusZoomNeighbors.PreviousIndex];
-				const FRichCurveKey Curve1Key1 = Curve1Keys[NextFocusZoomNeighbors.NextIndex];
-
-				Params.EvalTime = InZoom;
-				Params.Curve0Key0Time = Curve0Key0.Time;
-				Params.Curve0Key1Time = Curve0Key1.Time;
-				Params.Curve0Key0Tangent = Curve0Key0.LeaveTangent;
-				Params.Curve0Key1Tangent = Curve0Key1.ArriveTangent;
-				Params.Curve1Key0Time = Curve1Key0.Time;
-				Params.Curve1Key1Time = Curve1Key1.Time;
-				Params.Curve1Key0Tangent = Curve1Key0.LeaveTangent;
-				Params.Curve1Key1Tangent = Curve1Key1.ArriveTangent;
-
-				//Interpolate distortion parameters along the map blending curves, then linearly interpolate between focus points
-				const float Times[4] = { Curve0Key0.Time, Curve0Key1.Time, Curve1Key0.Time, Curve1Key1.Time };
-				const float Tangents[4] = { Curve0Key0.LeaveTangent, Curve0Key1.ArriveTangent, Curve1Key0.LeaveTangent, Curve1Key1.ArriveTangent };
-				const TConstArrayView<float> Parameters[4] = { PreviousFocusPreviousZoomPoint.DistortionInfo.Parameters
-															 , PreviousFocusNextZoomPoint.DistortionInfo.Parameters
-															 , NextFocusPreviousZoomPoint.DistortionInfo.Parameters
-															 , NextFocusNextZoomPoint.DistortionInfo.Parameters};
-				LensFileUtils::FindWeightsAndInterp(InZoom, Times, Tangents, FocusBlendFactor, Parameters, State.DistortionInfo.Parameters);
-
-				//Interpolate overscan factor along the map blending curves, then linearly interpolate between focus points
-				const float OverscanFactors[4] = { BlendedOverscanFactors[0], BlendedOverscanFactors[1], BlendedOverscanFactors[2], BlendedOverscanFactors[3] };
-				LensFileUtils::FindWeightsAndInterp(InZoom, Times, Tangents, FocusBlendFactor, OverscanFactors, InterpolatedOverscanFactor);
-			}
-		}
+	InterpolatedState.DistortionInfo.Parameters = Results.BlendedDistortionParams->Parameters;
+	
+	//Sets final blended distortion state
+	InLensHandler->SetDistortionState(InterpolatedState);
+	InLensHandler->SetCameraFilmback(CameraFilmback);
 
-		//Resulting displacement maps is the result of blending each individual maps together
-		//Distortion state can also provide distortion parameters and FxFy that were used.
-		//Instead of providing nothing (that accurately matches the blended map)
-		//We fill it with blended parameters and FxFy
+	//Draw resulting undistortion displacement map for evaluation point
+	LensFileRendering::DrawBlendedDisplacementMap(InLensHandler->GetUndistortionDisplacementMap()
+		, Results.BlendingParams.GetValue()
+		, UndistortionDisplacementMapHolders[0]
+		, UndistortionDisplacementMapHolders[1]
+		, UndistortionDisplacementMapHolders[2]
+		, UndistortionDisplacementMapHolders[3]);
 
-		FFocalLengthInfo FocalLength;
-		EvaluateFocalLength(InFocus, InZoom, FocalLength);
-		const FVector2D FxFyScale = FVector2D(LensInfo.SensorDimensions.X / InFilmback.X, LensInfo.SensorDimensions.Y / InFilmback.Y);
-		State.FocalLengthInfo.FxFy = FocalLength.FxFy * FxFyScale;
-		
-		//Sets final blended distortion state
-		InLensHandler->SetDistortionState(State);
-		InLensHandler->SetCameraFilmback(CameraFilmback);
+	//Draw resulting distortion displacement map for evaluation point
+	LensFileRendering::DrawBlendedDisplacementMap(InLensHandler->GetDistortionDisplacementMap()
+		, Results.BlendingParams.GetValue()
+		, DistortionDisplacementMapHolders[0]
+		, DistortionDisplacementMapHolders[1]
+		, DistortionDisplacementMapHolders[2]
+		, DistortionDisplacementMapHolders[3]);
 
-		//Draw resulting undistortion displacement map for evaluation point
-		LensFileRendering::DrawBlendedDisplacementMap(InLensHandler->GetUndistortionDisplacementMap()
-			, Params
-			, UndistortionDisplacementMapHolders[0]
-			, UndistortionDisplacementMapHolders[1]
-			, UndistortionDisplacementMapHolders[2]
-			, UndistortionDisplacementMapHolders[3]);
-
-		//Draw resulting distortion displacement map for evaluation point
-		LensFileRendering::DrawBlendedDisplacementMap(InLensHandler->GetDistortionDisplacementMap()
-			, Params
-			, DistortionDisplacementMapHolders[0]
-			, DistortionDisplacementMapHolders[1]
-			, DistortionDisplacementMapHolders[2]
-			, DistortionDisplacementMapHolders[3]);
-
-		InLensHandler->SetOverscanFactor(InterpolatedOverscanFactor);
-	}
+	InLensHandler->SetOverscanFactor(Results.BlendedOverscan.GetValue());
 	
 	return true;
 }
@@ -716,14 +526,7 @@ bool ULensFile::EvaluateDistortionForSTMaps(float InFocus, float InZoom, FVector
 		SetupNoDistortionOutput(InLensHandler);
 		return true;
 	}
-
-	if(STMapTable.FocusPoints.Num() <= 0)
-	{
-		UE_LOG(LogCameraCalibrationCore, Verbose, TEXT("Can't evaluate LensFile '%s' - No calibrated maps"), *GetName());
-		SetupNoDistortionOutput(InLensHandler);
-		return true;
-	}
-
+	
 	if (((LensInfo.SensorDimensions.X + UE_DOUBLE_KINDA_SMALL_NUMBER) < InFilmback.X) || ((LensInfo.SensorDimensions.Y + UE_DOUBLE_KINDA_SMALL_NUMBER) < InFilmback.Y))
 	{
 		UE_LOG(LogCameraCalibrationCore, Verbose
@@ -737,181 +540,82 @@ bool ULensFile::EvaluateDistortionForSTMaps(float InFocus, float InZoom, FVector
 	FCameraFilmbackSettings CameraFilmback;
 	CameraFilmback.SensorWidth = InFilmback.X;
 	CameraFilmback.SensorHeight = InFilmback.Y;
-
-	FDisplacementMapBlendingParams Params;
-
-	TArray<UTextureRenderTarget2D*, TInlineAllocator<4>> UndistortionMapSource;
-	UndistortionMapSource.AddZeroed(4);
-	TArray<UTextureRenderTarget2D*, TInlineAllocator<4>> DistortionMapSource;
-	DistortionMapSource.AddZeroed(4);
-
-	//When dealing with STMaps, FxFy was not a calibrated value. We can interpolate our curve directly for desired point
-	FLensDistortionState State;
-
+	
 	const FVector2D FxFyScale = FVector2D(InFilmback.X / LensInfo.SensorDimensions.X, InFilmback.Y / LensInfo.SensorDimensions.Y);
-	Params.FxFyScale = FxFyScale;
-
+	
+	//When dealing with STMaps, FxFy was not a calibrated value. We can interpolate our curve directly for desired point
 	FFocalLengthInfo FocalLength;
 	EvaluateFocalLength(InFocus, InZoom, FocalLength);
-	State.FocalLengthInfo.FxFy = FocalLength.FxFy * FxFyScale;
 
 	FImageCenterInfo ImageCenter;
 	EvaluateImageCenterParameters(InFocus, InZoom, ImageCenter);
-	State.ImageCenter = ImageCenter;
-	Params.PrincipalPoint = ImageCenter.PrincipalPoint;
-
-	float InterpolatedOverscanFactor = 0.0f;
 
-	//Find focuses in play
-	const LensDataTableUtils::FPointNeighbors Neighbors = LensDataTableUtils::FindFocusPoints(InFocus, STMapTable.GetFocusPoints());
-	const FSTMapFocusPoint& PreviousPoint = STMapTable.FocusPoints[Neighbors.PreviousIndex];
-	const FSTMapFocusPoint& NextPoint = STMapTable.FocusPoints[Neighbors.NextIndex];
-	if (Neighbors.NextIndex == Neighbors.PreviousIndex)
-	{
-		const LensDataTableUtils::FPointNeighbors ZoomNeighbors = LensDataTableUtils::FindZoomPoints(InZoom, PreviousPoint.ZoomPoints);
-		const FSTMapZoomPoint& PreviousZoomPoint = PreviousPoint.ZoomPoints[ZoomNeighbors.PreviousIndex];
-		const FSTMapZoomPoint& NextZoomPoint = PreviousPoint.ZoomPoints[ZoomNeighbors.NextIndex];
-
-		if (ZoomNeighbors.NextIndex == ZoomNeighbors.PreviousIndex)
-		{
-			//Exactly on one point
-			Params.BlendType = EDisplacementMapBlendType::OneFocusOneZoom;
+	LensInterpolationUtils::FDistortionMapBlendParams<FSTMapTable> Params;
+	Params.bGenerateBlendingParams = true;
 
-			UndistortionMapSource[0] = PreviousZoomPoint.DerivedDistortionData.UndistortionDisplacementMap;
-			DistortionMapSource[0] = PreviousZoomPoint.DerivedDistortionData.DistortionDisplacementMap;
- 			InterpolatedOverscanFactor = ComputeOverscan(PreviousZoomPoint.DerivedDistortionData.DistortionData, ImageCenter.PrincipalPoint);
-		}
-		else
+	// Callback that retrieves the displacement map render targets for the specified focus and zoom
+	Params.GetDisplacementMaps = LensInterpolationUtils::FDistortionMapBlendParams<FSTMapTable>::FGetDisplacementMaps::CreateLambda(
+		[](const FSTMapFocusPoint& FocusPoint, const FSTMapFocusCurve& FocusCurve, UTextureRenderTarget2D*& OutUndistortedMap, UTextureRenderTarget2D*& OutDistortedMap)
 		{
-			//Interpolate between two zoom points following the map blending curve
-			Params.BlendType = EDisplacementMapBlendType::OneFocusTwoZoom;
-
-			UndistortionMapSource[0] = PreviousZoomPoint.DerivedDistortionData.UndistortionDisplacementMap;
-			DistortionMapSource[0] = PreviousZoomPoint.DerivedDistortionData.DistortionDisplacementMap;
-			UndistortionMapSource[1] = NextZoomPoint.DerivedDistortionData.UndistortionDisplacementMap;
-			DistortionMapSource[1] = NextZoomPoint.DerivedDistortionData.DistortionDisplacementMap;
-
-			//Set displacement map blending parameters
-			const TArray<FRichCurveKey>& PreviousPointsKeys = PreviousPoint.MapBlendingCurve.GetConstRefOfKeys();
-			const FRichCurveKey Curve0Key0 = PreviousPointsKeys[ZoomNeighbors.PreviousIndex];
-			const FRichCurveKey Curve0Key1 = PreviousPointsKeys[ZoomNeighbors.NextIndex];
-
-			Params.EvalTime = InZoom;
-			Params.Curve0Key0Time = Curve0Key0.Time;
-			Params.Curve0Key1Time = Curve0Key1.Time;
-			Params.Curve0Key0Tangent = Curve0Key0.LeaveTangent;
-			Params.Curve0Key1Tangent = Curve0Key1.ArriveTangent;
-
-			//Interpolate overscan factor along the map blending curve
-			const float PreviousZoomPointOverscanFactor = ComputeOverscan(PreviousZoomPoint.DerivedDistortionData.DistortionData, ImageCenter.PrincipalPoint);
-			const float NextZoomPointOverscanFactor = ComputeOverscan(NextZoomPoint.DerivedDistortionData.DistortionData, ImageCenter.PrincipalPoint);
-
-			const float Times[2] = { Curve0Key0.Time, Curve0Key1.Time };
-			const float Tangents[2] = { Curve0Key0.LeaveTangent, Curve0Key1.ArriveTangent };
-			const float OverscanFactors[2] = { PreviousZoomPointOverscanFactor, NextZoomPointOverscanFactor };
-			LensFileUtils::FindWeightsAndInterp(InZoom, Times, Tangents, TOptional<float>(), OverscanFactors, InterpolatedOverscanFactor);
-		}
-	}
-	else
-	{
-		//Previous Focus two zoom points
-		const LensDataTableUtils::FPointNeighbors PreviousFocusZoomNeighbors = LensDataTableUtils::FindZoomPoints(InZoom, PreviousPoint.ZoomPoints);
-		const FSTMapZoomPoint& PreviousFocusPreviousZoomPoint = PreviousPoint.ZoomPoints[PreviousFocusZoomNeighbors.PreviousIndex];
-		const FSTMapZoomPoint& PreviousFocusNextZoomPoint = PreviousPoint.ZoomPoints[PreviousFocusZoomNeighbors.NextIndex];
-
-		//Next focus two zoom points
-		const LensDataTableUtils::FPointNeighbors NextFocusZoomNeighbors = LensDataTableUtils::FindZoomPoints(InZoom, NextPoint.ZoomPoints);
-		const FSTMapZoomPoint NextFocusPreviousZoomPoint = NextPoint.ZoomPoints[NextFocusZoomNeighbors.PreviousIndex];
-		const FSTMapZoomPoint& NextFocusNextZoomPoint = NextPoint.ZoomPoints[NextFocusZoomNeighbors.NextIndex];
-		
-		const float FocusBlendFactor = LensInterpolationUtils::GetBlendFactor(InFocus, PreviousPoint.Focus, NextPoint.Focus);
-		Params.FocusBlendFactor = FocusBlendFactor;
+			if (const FSTMapZoomPoint* ZoomPoint = FocusPoint.GetZoomPoint(FocusCurve.Zoom))
+			{
+				OutUndistortedMap = ZoomPoint->DerivedDistortionData.UndistortionDisplacementMap;
+				OutDistortedMap = ZoomPoint->DerivedDistortionData.UndistortionDisplacementMap;
+			}
+		});
 
-		//Verify if we are dealing with one zoom point on each focus. If that's the case, we are doing simple lerp across the focus curves
-		if(PreviousFocusZoomNeighbors.NextIndex == PreviousFocusZoomNeighbors.PreviousIndex
-			&& NextFocusZoomNeighbors.NextIndex == NextFocusZoomNeighbors.PreviousIndex)
+	// Callback when the blend function constructs the displacement maps for each corner used in the blend, which generates the displacement maps
+	// and returns the computed overscan of the displacement map
+	Params.ProcessDisplacementMaps = LensInterpolationUtils::FDistortionMapBlendParams<FSTMapTable>::FProcessDisplacementMaps::CreateLambda(
+		[this, &ImageCenter]
+		(const FSTMapFocusPoint& FocusPoint, const FSTMapFocusCurve& FocusCurve, UTextureRenderTarget2D*, UTextureRenderTarget2D*)
 		{
-			Params.BlendType = EDisplacementMapBlendType::TwoFocusOneZoom;
-
-			UndistortionMapSource[0] = PreviousFocusPreviousZoomPoint.DerivedDistortionData.UndistortionDisplacementMap;
-			DistortionMapSource[0] = PreviousFocusPreviousZoomPoint.DerivedDistortionData.DistortionDisplacementMap;
-			UndistortionMapSource[1] = NextFocusPreviousZoomPoint.DerivedDistortionData.UndistortionDisplacementMap;
-			DistortionMapSource[1] = NextFocusPreviousZoomPoint.DerivedDistortionData.DistortionDisplacementMap;
+			if (const FSTMapZoomPoint* ZoomPoint = FocusPoint.GetZoomPoint(FocusCurve.Zoom))
+			{
+				return ComputeOverscan(ZoomPoint->DerivedDistortionData.DistortionData, ImageCenter.PrincipalPoint);
+			}
 
-			//Linearly interpolate the overscan factor
-			const float PreviousPointOverscanFactor = ComputeOverscan(PreviousFocusPreviousZoomPoint.DerivedDistortionData.DistortionData, ImageCenter.PrincipalPoint);
-			const float NextPointOverscanFactor = ComputeOverscan(NextFocusPreviousZoomPoint.DerivedDistortionData.DistortionData, ImageCenter.PrincipalPoint);
-			InterpolatedOverscanFactor = FMath::Lerp(PreviousPointOverscanFactor, NextPointOverscanFactor, FocusBlendFactor);
-		}
-		else
-		{
-			//Interpolate between two zoom points on each curve using the correct map blending curve, then linearly interpolate between the two focus points
-			Params.BlendType = EDisplacementMapBlendType::TwoFocusTwoZoom;
-
-			UndistortionMapSource[0] = PreviousFocusPreviousZoomPoint.DerivedDistortionData.UndistortionDisplacementMap;
-			DistortionMapSource[0] = PreviousFocusPreviousZoomPoint.DerivedDistortionData.DistortionDisplacementMap;
-			UndistortionMapSource[1] = PreviousFocusNextZoomPoint.DerivedDistortionData.UndistortionDisplacementMap;
-			DistortionMapSource[1] = PreviousFocusNextZoomPoint.DerivedDistortionData.DistortionDisplacementMap;
-			UndistortionMapSource[2] = NextFocusPreviousZoomPoint.DerivedDistortionData.UndistortionDisplacementMap;
-			DistortionMapSource[2] = NextFocusPreviousZoomPoint.DerivedDistortionData.DistortionDisplacementMap;
-			UndistortionMapSource[3] = NextFocusNextZoomPoint.DerivedDistortionData.UndistortionDisplacementMap;
-			DistortionMapSource[3] = NextFocusNextZoomPoint.DerivedDistortionData.DistortionDisplacementMap;
-			
-			//Set displacement map blending parameters
-			const TArray<FRichCurveKey>& Curve0Keys = PreviousPoint.MapBlendingCurve.GetConstRefOfKeys();
-			const FRichCurveKey Curve0Key0 = Curve0Keys[PreviousFocusZoomNeighbors.PreviousIndex];
-			const FRichCurveKey Curve0Key1 = Curve0Keys[PreviousFocusZoomNeighbors.NextIndex];
-
-			const TArray<FRichCurveKey>& Curve1Keys = NextPoint.MapBlendingCurve.GetConstRefOfKeys();
-			const FRichCurveKey Curve1Key0 = Curve1Keys[NextFocusZoomNeighbors.PreviousIndex];
-			const FRichCurveKey Curve1Key1 = Curve1Keys[NextFocusZoomNeighbors.NextIndex];
-
-			Params.EvalTime = InZoom;
-			Params.Curve0Key0Time = Curve0Key0.Time;
-			Params.Curve0Key1Time = Curve0Key1.Time;
-			Params.Curve0Key0Tangent = Curve0Key0.LeaveTangent;
-			Params.Curve0Key1Tangent = Curve0Key1.ArriveTangent;
-			Params.Curve1Key0Time = Curve1Key0.Time;
-			Params.Curve1Key1Time = Curve1Key1.Time;
-			Params.Curve1Key0Tangent = Curve1Key0.LeaveTangent;
-			Params.Curve1Key1Tangent = Curve1Key1.ArriveTangent;
-
-			//Interpolate overscan factor along the map blending curves, then linearly interpolate between focus points
-			const float PreviousFocusPreviousZoomOverscanFactor = ComputeOverscan(PreviousFocusPreviousZoomPoint.DerivedDistortionData.DistortionData, ImageCenter.PrincipalPoint);
-			const float PreviousFocusNextZoomOverscanFactor = ComputeOverscan(PreviousFocusNextZoomPoint.DerivedDistortionData.DistortionData, ImageCenter.PrincipalPoint);
-			const float NextFocusPreviousZoomOverscanFactor = ComputeOverscan(NextFocusPreviousZoomPoint.DerivedDistortionData.DistortionData, ImageCenter.PrincipalPoint);
-			const float NextFocusNextZoomOverscanFactor = ComputeOverscan(NextFocusNextZoomPoint.DerivedDistortionData.DistortionData, ImageCenter.PrincipalPoint);
-
-			const float Times[4] = { Curve0Key0.Time, Curve0Key1.Time, Curve1Key0.Time, Curve1Key1.Time };
-			const float Tangents[4] = { Curve0Key0.LeaveTangent, Curve0Key1.ArriveTangent, Curve1Key0.LeaveTangent, Curve1Key1.ArriveTangent };
-			const float OverscanFactors[4] = { PreviousFocusPreviousZoomOverscanFactor, PreviousFocusNextZoomOverscanFactor, NextFocusPreviousZoomOverscanFactor, NextFocusNextZoomOverscanFactor };
-			LensFileUtils::FindWeightsAndInterp(InZoom, Times, Tangents, FocusBlendFactor, OverscanFactors, InterpolatedOverscanFactor);
-		}
+			return 1.0f;
+		});
+	
+	LensInterpolationUtils::FDistortionMapBlendResults Results = LensInterpolationUtils::DistortionMapBlend(STMapTable, InFocus, InZoom, Params);
+	if (!Results.bValid)
+	{
+		UE_LOG(LogCameraCalibrationCore, Verbose, TEXT("Can't evaluate LensFile '%s' - No calibrated maps"), *GetName());
+		SetupNoDistortionOutput(InLensHandler);
+		return true;
 	}
 
+	FDisplacementMapBlendingParams& BlendingParams = Results.BlendingParams.GetValue();
+	BlendingParams.FxFyScale = FxFyScale;
+	BlendingParams.PrincipalPoint = ImageCenter.PrincipalPoint;
+	
 	//Draw resulting undistortion displacement map for evaluation point
 	LensFileRendering::DrawBlendedDisplacementMap(InLensHandler->GetUndistortionDisplacementMap()
-		, Params
-		, UndistortionMapSource[0]
-		, UndistortionMapSource[1]
-		, UndistortionMapSource[2]
-		, UndistortionMapSource[3]);
+		, BlendingParams
+		, Results.UndistortedMaps.GetValue()[0]
+		, Results.UndistortedMaps.GetValue()[1]
+		, Results.UndistortedMaps.GetValue()[2]
+		, Results.UndistortedMaps.GetValue()[3]);
 
 	//Draw resulting distortion displacement map for evaluation point
 	LensFileRendering::DrawBlendedDisplacementMap(InLensHandler->GetDistortionDisplacementMap()
-		, Params
-		, DistortionMapSource[0]
-		, DistortionMapSource[1]
-		, DistortionMapSource[2]
-		, DistortionMapSource[3]);
-
-			
+		, BlendingParams
+		, Results.DistortedMaps.GetValue()[0]
+		, Results.DistortedMaps.GetValue()[1]
+		, Results.DistortedMaps.GetValue()[2]
+		, Results.DistortedMaps.GetValue()[3]);
+	
+	FLensDistortionState State;
+	State.FocalLengthInfo.FxFy = FocalLength.FxFy * FxFyScale;
+	State.ImageCenter = ImageCenter;
+	
 	//Sets final blended distortion state
 	InLensHandler->SetDistortionState(MoveTemp(State));
 	InLensHandler->SetCameraFilmback(CameraFilmback);
 
-	InLensHandler->SetOverscanFactor(InterpolatedOverscanFactor);
-
+	InLensHandler->SetOverscanFactor(Results.BlendedOverscan.GetValue());
+	
 	return true;
 }
 
@@ -919,49 +623,27 @@ bool ULensFile::EvaluateNodalPointOffset(float InFocus, float InZoom, FNodalPoin
 {
 	TRACE_CPUPROFILER_EVENT_SCOPE(ULensFile::EvaluateNodalPointOffset);
 
-	if (NodalOffsetTable.GetFocusPoints().Num() <= 0)
+	constexpr int32 NumParams = 6;
+	TArray<float> BlendedParameters;
+	if (LensInterpolationUtils::IndexedParameterBlend(NodalOffsetTable.FocusPoints, NodalOffsetTable.FocusCurves, InFocus, InZoom, NumParams, BlendedParameters))
 	{
-		return false;
-	}
+		ensure(BlendedParameters.Num() == NumParams);
 
-	const auto EvaluateNodalOffset = [this](int32 FocusIndex,float Zoom, FNodalPointOffset& OutNodalOffset)
-	{
-		FRotator RotationOffset;
-		constexpr int32 RotationDimensions = 3;
-		for(int32 AxisIndex = 0; AxisIndex < RotationDimensions; ++AxisIndex)
+		FVector Location;
+		FRotator Rotation;
+		
+		for (int32 Index = 0; Index < 3; ++Index)
 		{
-			const EAxis::Type Axis = static_cast<EAxis::Type>(AxisIndex + 1);
-			OutNodalOffset.LocationOffset.SetComponentForAxis(Axis, NodalOffsetTable.FocusPoints[FocusIndex].LocationOffset[AxisIndex].Eval(Zoom));
-			RotationOffset.SetComponentForAxis(Axis, NodalOffsetTable.FocusPoints[FocusIndex].RotationOffset[AxisIndex].Eval(Zoom));
+			Location[Index] = BlendedParameters[Index];
+			Rotation.SetComponentForAxis(static_cast<EAxis::Type>(Index + 1), BlendedParameters[Index + 3]);
 		}
-		
-		OutNodalOffset.RotationOffset = FQuat(RotationOffset);	
-	};
 
-	if (NodalOffsetTable.GetFocusPoints().Num() == 1)
-	{
-		EvaluateNodalOffset(0, InZoom, OutEvaluatedValue);
+		OutEvaluatedValue.LocationOffset = Location;
+		OutEvaluatedValue.RotationOffset = FQuat(Rotation);
 		return true;
 	}
 	
-	const LensDataTableUtils::FPointNeighbors Neighbors = LensDataTableUtils::FindFocusPoints(InFocus, NodalOffsetTable.GetFocusPoints());
-	if (Neighbors.NextIndex == Neighbors.PreviousIndex)
-	{
-		EvaluateNodalOffset(Neighbors.PreviousIndex, InZoom, OutEvaluatedValue);
-	}
-	else
-	{
-		FNodalPointOffset PreviousValue;
-		EvaluateNodalOffset(Neighbors.PreviousIndex, InZoom, PreviousValue);
-		FNodalPointOffset NextValue;
-		EvaluateNodalOffset(Neighbors.NextIndex, InZoom, NextValue);
-	
-		//Blend result between focus
-		const float BlendFactor = LensInterpolationUtils::GetBlendFactor(InFocus, NodalOffsetTable.FocusPoints[Neighbors.PreviousIndex].Focus, NodalOffsetTable.FocusPoints[Neighbors.NextIndex].Focus);
-		LensInterpolationUtils::Interpolate(BlendFactor, &PreviousValue, &NextValue, &OutEvaluatedValue);
-	}
-
-	return true;
+	return false;
 }
 
 bool ULensFile::HasFocusEncoderMapping() const
@@ -1492,3 +1174,13 @@ ULensFile* FLensFilePicker::GetLensFile() const
 	}
 }
 
+#if WITH_EDITOR
+void ULensFile::BuildLensTableFocusCurves()
+{
+	FocalLengthTable.BuildFocusCurves();
+	DistortionTable.BuildFocusCurves();
+	ImageCenterTable.BuildFocusCurves();
+	STMapTable.BuildFocusCurves();
+	NodalOffsetTable.BuildFocusCurves();
+}
+#endif // WITH_EDITOR
diff --git a/Engine/Plugins/VirtualProduction/CameraCalibrationCore/Source/CameraCalibrationCore/Private/LensFileObjectVersion.cpp b/Engine/Plugins/VirtualProduction/CameraCalibrationCore/Source/CameraCalibrationCore/Private/LensFileObjectVersion.cpp
new file mode 100644
--- /dev/null
+++ b/Engine/Plugins/VirtualProduction/CameraCalibrationCore/Source/CameraCalibrationCore/Private/LensFileObjectVersion.cpp
@@ -0,0 +1,8 @@
+﻿// Copyright Epic Games, Inc. All Rights Reserved.
+
+#include "LensFileObjectVersion.h"
+
+#include "UObject/DevObjectVersion.h"
+
+const FGuid FLensFileObjectVersion::GUID(0x8652A554, 0x966A466C, 0x9FD71C6D, 0xD61B1ADB);
+static FDevVersionRegistration GRegisterLensFileObjectVersion(FLensFileObjectVersion::GUID, FLensFileObjectVersion::LatestVersion, TEXT("LensFileVersion"));
\ No newline at end of file
diff --git a/Engine/Plugins/VirtualProduction/CameraCalibrationCore/Source/CameraCalibrationCore/Private/LensFileRendering.cpp b/Engine/Plugins/VirtualProduction/CameraCalibrationCore/Source/CameraCalibrationCore/Private/LensFileRendering.cpp
--- a/Engine/Plugins/VirtualProduction/CameraCalibrationCore/Source/CameraCalibrationCore/Private/LensFileRendering.cpp
+++ b/Engine/Plugins/VirtualProduction/CameraCalibrationCore/Source/CameraCalibrationCore/Private/LensFileRendering.cpp
@@ -25,16 +25,9 @@ class FDisplacementMapBlendPS : public FGlobalShader
 	
 	BEGIN_SHADER_PARAMETER_STRUCT(FParameters, )
 
-		SHADER_PARAMETER(float, EvalTime)
-		SHADER_PARAMETER(float, Curve0Key0Time)
-		SHADER_PARAMETER(float, Curve0Key1Time)
-		SHADER_PARAMETER(float, Curve0Key0Tangent)
-		SHADER_PARAMETER(float, Curve0Key1Tangent)
-		SHADER_PARAMETER(float, Curve1Key0Time)
-		SHADER_PARAMETER(float, Curve1Key1Time)
-		SHADER_PARAMETER(float, Curve1Key0Tangent)
-		SHADER_PARAMETER(float, Curve1Key1Tangent)
-		SHADER_PARAMETER(float, FocusBlendFactor)
+		SHADER_PARAMETER(float, EvalFocus)
+		SHADER_PARAMETER(float, EvalZoom)
+		SHADER_PARAMETER_ARRAY(FVector4f, PatchCorners, [4])
 		SHADER_PARAMETER(FVector2f, FxFyScale)
 		SHADER_PARAMETER(FVector2f, PrincipalPoint)
 		SHADER_PARAMETER(FIntPoint, OutputTextureExtent)
@@ -151,7 +144,12 @@ namespace LensFileRendering
 					PermutationVector.Set<FDisplacementMapBlendPS::FBlendType>(1);
 					check(SourceTextureTwoResource);
 					const FRDGTextureRef TextureTwo = GraphBuilder.RegisterExternalTexture(CreateRenderTarget(SourceTextureTwoResource->TextureRHI, TEXT("DisplacementMapTwo")));
-					PassParameters->FocusBlendFactor = BlendParams.FocusBlendFactor;
+					PassParameters->EvalFocus = BlendParams.EvalFocus;
+					PassParameters->EvalZoom = BlendParams.EvalZoom;
+					PassParameters->PatchCorners[0] = BlendParams.PatchCorners[0].ToVector();
+					PassParameters->PatchCorners[1] = BlendParams.PatchCorners[1].ToVector();
+					PassParameters->PatchCorners[2] = FVector4f::Zero();
+					PassParameters->PatchCorners[3] = FVector4f::Zero();
 					PassParameters->SourceTextureTwo = TextureTwo;
 					break;
 				}
@@ -160,11 +158,12 @@ namespace LensFileRendering
 					PermutationVector.Set<FDisplacementMapBlendPS::FBlendType>(2);
 					check(SourceTextureTwoResource);
 					const FRDGTextureRef TextureTwo = GraphBuilder.RegisterExternalTexture(CreateRenderTarget(SourceTextureTwoResource->TextureRHI, TEXT("DisplacementMapTwo")));
-					PassParameters->EvalTime = BlendParams.EvalTime;
-					PassParameters->Curve0Key0Time = BlendParams.Curve0Key0Time;
-					PassParameters->Curve0Key1Time = BlendParams.Curve0Key1Time;
-					PassParameters->Curve0Key0Tangent = BlendParams.Curve0Key0Tangent;
-					PassParameters->Curve0Key1Tangent = BlendParams.Curve0Key1Tangent;
+					PassParameters->EvalFocus = BlendParams.EvalFocus;
+					PassParameters->EvalZoom = BlendParams.EvalZoom;
+					PassParameters->PatchCorners[0] = BlendParams.PatchCorners[0].ToVector();
+					PassParameters->PatchCorners[1] = BlendParams.PatchCorners[1].ToVector();
+					PassParameters->PatchCorners[2] = FVector4f::Zero();
+					PassParameters->PatchCorners[3] = FVector4f::Zero();
 					PassParameters->SourceTextureTwo = TextureTwo;
 					break;
 				}
@@ -175,16 +174,12 @@ namespace LensFileRendering
 					const FRDGTextureRef TextureTwo = GraphBuilder.RegisterExternalTexture(CreateRenderTarget(SourceTextureTwoResource->TextureRHI, TEXT("DisplacementMapTwo")));
 					const FRDGTextureRef TextureThree = GraphBuilder.RegisterExternalTexture(CreateRenderTarget(SourceTextureThreeResource->TextureRHI, TEXT("DisplacementMapThree")));
 					const FRDGTextureRef TextureFour = GraphBuilder.RegisterExternalTexture(CreateRenderTarget(SourceTextureFourResource->TextureRHI, TEXT("DisplacementMapFour")));
-					PassParameters->EvalTime = BlendParams.EvalTime;
-					PassParameters->Curve0Key0Time = BlendParams.Curve0Key0Time;
-					PassParameters->Curve0Key1Time = BlendParams.Curve0Key1Time;
-					PassParameters->Curve0Key0Tangent = BlendParams.Curve0Key0Tangent;
-					PassParameters->Curve0Key1Tangent = BlendParams.Curve0Key1Tangent;
-					PassParameters->Curve1Key0Time = BlendParams.Curve1Key0Time;
-					PassParameters->Curve1Key1Time = BlendParams.Curve1Key1Time;
-					PassParameters->Curve1Key0Tangent = BlendParams.Curve1Key0Tangent;
-					PassParameters->Curve1Key1Tangent = BlendParams.Curve1Key1Tangent;
-					PassParameters->FocusBlendFactor = BlendParams.FocusBlendFactor;
+					PassParameters->EvalFocus = BlendParams.EvalFocus;
+					PassParameters->EvalZoom = BlendParams.EvalZoom;
+					PassParameters->PatchCorners[0] = BlendParams.PatchCorners[0].ToVector();
+					PassParameters->PatchCorners[1] = BlendParams.PatchCorners[1].ToVector();
+					PassParameters->PatchCorners[2] = BlendParams.PatchCorners[2].ToVector();
+					PassParameters->PatchCorners[3] = BlendParams.PatchCorners[3].ToVector();
 					PassParameters->SourceTextureTwo = TextureTwo;
 					PassParameters->SourceTextureThree = TextureThree;
 					PassParameters->SourceTextureFour = TextureFour;
diff --git a/Engine/Plugins/VirtualProduction/CameraCalibrationCore/Source/CameraCalibrationCore/Private/LensFileRendering.h b/Engine/Plugins/VirtualProduction/CameraCalibrationCore/Source/CameraCalibrationCore/Private/LensFileRendering.h
--- a/Engine/Plugins/VirtualProduction/CameraCalibrationCore/Source/CameraCalibrationCore/Private/LensFileRendering.h
+++ b/Engine/Plugins/VirtualProduction/CameraCalibrationCore/Source/CameraCalibrationCore/Private/LensFileRendering.h
@@ -4,6 +4,7 @@
 
 #include "CoreTypes.h"
 #include "Math/Vector2D.h"
+#include "Math/Vector4.h"
 
 class UTextureRenderTarget2D;
 
@@ -16,6 +17,36 @@ enum class EDisplacementMapBlendType :uint8
 	TwoFocusTwoZoom, //Two Besier interp between each pair of zoom points and one linear interp across focus
 };
 
+/** Corner of the blending patch which stores x and y values and x and y tangents */
+struct FDisplacementMapBlendPatchCorner
+{
+	FDisplacementMapBlendPatchCorner(float InX, float InY, float InTangentX, float InTangentY)
+		: X(InX)
+		, Y(InY)
+		, TangentX(InTangentX)
+		, TangentY(InTangentY)
+	{ }
+
+	FDisplacementMapBlendPatchCorner()
+		: FDisplacementMapBlendPatchCorner(0.0f, 0.0f, 0.0f, 0.0f)
+	{ }
+	
+	/** Converts the patch corner values to a vector */
+	FVector4f ToVector() const { return FVector4f(X, Y, TangentX, TangentY); }
+	
+	/** X coordinate of the corner */
+	float X;
+
+	/** Y coordinate of the corner */
+	float Y;
+
+	/** Tangent in the x direction of the corner */
+	float TangentX;
+
+	/** Tangent in the y direction of the corner */
+	float TangentY;
+};
+
 /** Single struct containing blending params for all types */
 struct FDisplacementMapBlendingParams
 {
@@ -23,20 +54,12 @@ struct FDisplacementMapBlendingParams
 	EDisplacementMapBlendType BlendType = EDisplacementMapBlendType::OneFocusOneZoom;
 
 	/** Bezier blend parameters */
-	float EvalTime = 0.0f;
-
-	float Curve0Key0Time = 0.0f;
-	float Curve0Key1Time = 0.0f;
-	float Curve0Key0Tangent = 0.0f;
-	float Curve0Key1Tangent = 0.0f;
-
-	float Curve1Key0Time = 0.0f;
-	float Curve1Key1Time = 0.0f;
-	float Curve1Key0Tangent = 0.0f;
-	float Curve1Key1Tangent = 0.0f;
-
-	float FocusBlendFactor = 0.0f;
+	float EvalFocus = 0.0f;
+	float EvalZoom = 0.0f;
 
+	/** Corners of the blending patch, indexed in the following order: (X0, Y0) -> (X1, Y0) -> (X1, Y1) -> (X0, Y1) */
+	FDisplacementMapBlendPatchCorner PatchCorners[4];
+	
 	/** Scale parameter that allows displacement maps for one sensor size to be applied to camera's with a different sensor size */
 	FVector2D FxFyScale = { 1.0f, 1.0f };
 
diff --git a/Engine/Plugins/VirtualProduction/CameraCalibrationCore/Source/CameraCalibrationCore/Private/LensInterpolationUtils.cpp b/Engine/Plugins/VirtualProduction/CameraCalibrationCore/Source/CameraCalibrationCore/Private/LensInterpolationUtils.cpp
--- a/Engine/Plugins/VirtualProduction/CameraCalibrationCore/Source/CameraCalibrationCore/Private/LensInterpolationUtils.cpp
+++ b/Engine/Plugins/VirtualProduction/CameraCalibrationCore/Source/CameraCalibrationCore/Private/LensInterpolationUtils.cpp
@@ -4,6 +4,8 @@
 #include "LensInterpolationUtils.h"
 
 #include "LensData.h"
+#include "Curves/CurveEvaluation.h"
+#include "Curves/RichCurve.h"
 
 
 //Property interpolation utils largely inspired from livelink interp code
@@ -144,6 +146,91 @@ namespace LensInterpolationUtils
 			return 1.0f;
 		}
 	}
+
+	float FTangentBezierCurve::Eval(float InX)
+	{
+		if (FMath::IsNearlyEqual(X0, X1))
+		{
+			return Y0;
+		}
+
+		constexpr float OneThird = 1.0f / 3.0f;
+
+		const float DeltaX = X1 - X0;
+		const float DeltaY = Y1 - Y0;
+		const float Alpha = (InX - X0) / DeltaX;
+		const float TangentScale = DeltaY / DeltaX;
+		
+		const float P0 = Y0;
+		const float P3 = Y1;
+		const float P1 = P0 + (Tangent0 * TangentScale * DeltaX * OneThird);
+		const float P2 = P3 - (Tangent1 * TangentScale * DeltaX * OneThird);
+		return UE::Curves::BezierInterp(P0, P1, P2, P3, Alpha);
+	}
+
+	float FCoonsPatch::Blend(const FVector2D& InPoint)
+	{
+		float Alpha = 0.0;
+		if (!FMath::IsNearlyEqual(X0, X1))
+		{
+			Alpha = (InPoint.X - X0) / (X1 - X0);
+		}
+		
+		float Beta = 0.0;
+		if (!FMath::IsNearlyEqual(Y0, Y1))
+		{
+			Beta = (InPoint.Y - Y0) / (Y1 - Y0);
+		}
+		
+		// In degenerate cases (such as only having 3 defined corners), the four curves' corners may not match.
+		// Simply use the average of the two possible values in that case; while this doesn't match the strict
+		// definition for a Coons patch, it allows the patch to give a somewhat usable value for such edge cases
+		const float P00 = X0Curve.Eval(X0) + Y0Curve.Eval(Y0);
+		const float P01 = X0Curve.Eval(X1) + Y1Curve.Eval(Y0);
+		const float P10 = X1Curve.Eval(X0) + Y0Curve.Eval(Y1);
+		const float P11 = X1Curve.Eval(X1) + Y1Curve.Eval(Y1);
+			
+		const float Lx = FMath::Lerp(X0Curve.Eval(InPoint.X), X1Curve.Eval(InPoint.X), Beta);
+		const float Ly = FMath::Lerp(Y0Curve.Eval(InPoint.Y), Y1Curve.Eval(InPoint.Y), Alpha);
+		const float B = FMath::BiLerp(P00, P01, P10, P11, Alpha, Beta);
+		
+		return Lx + Ly - 0.5 * B;
+	}
+
+	FTangentBezierCoonsPatch::FTangentBezierCoonsPatch(float X0, float X1, float Y0, float Y1, float XTangents[4], float YTangents[4])
+		: Corners {
+			FPatchCorner(X0, Y0, XTangents[0], YTangents[0], 0.0),
+			FPatchCorner(X1, Y0, XTangents[1], YTangents[1], 0.0),
+			FPatchCorner(X1, Y1, XTangents[2], YTangents[2], 0.0),
+			FPatchCorner(X0, Y1, XTangents[3], YTangents[3], 0.0)}
+	{ }
+
+	float FTangentBezierCoonsPatch::Blend(const FVector2D& InPoint)
+	{
+		
+		float Alpha = 0.0;
+		if (!FMath::IsNearlyEqual(Corners[0].X, Corners[1].X))
+		{
+			Alpha = (InPoint.X - Corners[0].X) / (Corners[1].X - Corners[0].X);
+		}
+		
+		float Beta = 0.0;
+		if (!FMath::IsNearlyEqual(Corners[0].Y, Corners[3].Y))
+		{
+			Beta = (InPoint.Y - Corners[0].Y) / (Corners[3].Y - Corners[0].Y);
+		}
+
+		FTangentBezierCurve X0Curve = FTangentBezierCurve(Corners[0].X, Corners[1].X, Corners[0].Value, Corners[1].Value, Corners[0].TangentX, Corners[1].TangentX);
+		FTangentBezierCurve X1Curve = FTangentBezierCurve(Corners[3].X, Corners[2].X, Corners[3].Value, Corners[2].Value, Corners[3].TangentX, Corners[2].TangentX);
+		FTangentBezierCurve Y0Curve = FTangentBezierCurve(Corners[0].Y, Corners[3].Y, Corners[0].Value, Corners[3].Value, Corners[0].TangentY, Corners[3].TangentY);
+		FTangentBezierCurve Y1Curve = FTangentBezierCurve(Corners[1].Y, Corners[2].Y, Corners[1].Value, Corners[2].Value, Corners[1].TangentY, Corners[2].TangentY);
+	
+		const float Lx = FMath::Lerp(X0Curve.Eval(InPoint.X), X1Curve.Eval(InPoint.X), Beta);
+		const float Ly = FMath::Lerp(Y0Curve.Eval(InPoint.Y), Y1Curve.Eval(InPoint.Y), Alpha);
+		const float B = FMath::BiLerp(Corners[0].Value, Corners[1].Value, Corners[3].Value, Corners[2].Value, Alpha, Beta);
+		
+		return Lx + Ly - B;
+	}
 }
 
 
diff --git a/Engine/Plugins/VirtualProduction/CameraCalibrationCore/Source/CameraCalibrationCore/Private/LensInterpolationUtils.h b/Engine/Plugins/VirtualProduction/CameraCalibrationCore/Source/CameraCalibrationCore/Private/LensInterpolationUtils.h
--- a/Engine/Plugins/VirtualProduction/CameraCalibrationCore/Source/CameraCalibrationCore/Private/LensInterpolationUtils.h
+++ b/Engine/Plugins/VirtualProduction/CameraCalibrationCore/Source/CameraCalibrationCore/Private/LensInterpolationUtils.h
@@ -3,12 +3,19 @@
 #pragma once
 
 #include "CoreTypes.h"
+#include "LensFileRendering.h"
+#include "Curves/RichCurve.h"
+#include "Math/Vector2D.h"
 
 #include "Misc/EnumClassFlags.h"
+#include "Tables/DistortionParametersTable.h"
+#include "Tables/LensTableUtils.h"
 #include "UObject/Object.h"
 #include "UObject/ObjectMacros.h"
 
 
+struct FRichCurve;
+
 namespace LensInterpolationUtils
 {
 	template<typename Type>
@@ -26,4 +33,549 @@ namespace LensInterpolationUtils
 	}
 	
 	float GetBlendFactor(float InValue, float ValueA, float ValueB);
+
+	/** A cubic bezier curve constructed from two points and their tangents */
+	struct FTangentBezierCurve
+	{
+		FTangentBezierCurve(float InX0, float InX1, float InY0, float InY1, float InTangent0, float InTangent1)
+			: X0(InX0)
+			, X1(InX1)
+			, Y0(InY0)
+			, Y1(InY1)
+			, Tangent0(InTangent0)
+			, Tangent1(InTangent1)
+
+		{}
+
+		/** Evaluates the Bezier curve at the specified x value */
+		float Eval(float InX);
+
+		/** Minimum x value of the curve */
+		float X0;
+
+		/** Maximum x value of the curve */
+		float X1;
+
+		/** Y value of the curve at x0 */
+		float Y0;
+
+		/** Y value of the curve at x1 */
+		float Y1;
+
+		/** Tangent of the curve at x0 */
+		float Tangent0;
+
+		/** Tangent of the curve at x1 */
+		float Tangent1;
+	};
+	
+	/**
+	 * A Coons patch (https://en.wikipedia.org/wiki/Coons_patch) for blending four boundary curves together
+	 * into a surface patch. Takes in four rich curves for each side, two for x-axis and two for y-axis, as well as
+	 * min and max values for x and y from which the corner points of the patch are formed.
+	 */
+	struct FCoonsPatch
+	{
+		FCoonsPatch(const FRichCurve& X0Curve, const FRichCurve& X1Curve, const FRichCurve& Y0Curve, const FRichCurve& Y1Curve, float X0, float X1, float Y0, float Y1)
+			: X0Curve(X0Curve)
+			, X1Curve(X1Curve)
+			, Y0Curve(Y0Curve)
+			, Y1Curve(Y1Curve)
+			, X0(X0)
+			, X1(X1)
+			, Y0(Y0)
+			, Y1(Y1)
+		{ }
+		
+		/** Computes the value of the patch at the specified point by blending between the four edge curves */
+		float Blend(const FVector2D& InPoint);
+		
+		/** Min x-axis curve of the patch */
+		const FRichCurve& X0Curve;
+
+		/** Max x-axis curve of the patch */
+		const FRichCurve& X1Curve;
+
+		/** Min y-axis curve of the patch */
+		const FRichCurve& Y0Curve;
+
+		/** Max y-axis curve of the patch */
+		const FRichCurve& Y1Curve;
+
+		/** Minimum x value of the patch */
+		float X0;
+
+		/** Maximum x value of the patch */
+		float X1;
+
+		/** Minimum y value of the patch */
+		float Y0;
+
+		/** Maximum y value of the patch */
+		float Y1;
+	};
+
+	/**
+	 * A version of the Coons patch that takes in four points that each have x and y tangent data, and performs
+	 * a Bezier interpolation to create the four edge curves from the points and their tangents
+	 */
+	struct FTangentBezierCoonsPatch
+	{
+		/** Creates a new Coons patch from four corner points with x and y tangents. */
+		FTangentBezierCoonsPatch(float X0, float X1, float Y0, float Y1, float XTangents[4], float YTangents[4]);
+		
+		/** Computes the value of the patch at the specified point by blending between the four edge curves */
+		float Blend(const FVector2D& InPoint);
+
+		/** Stores the x and y position and x and y tangents of a corner */
+		struct FPatchCorner
+		{
+			float X;
+			float Y;
+			float TangentX;
+			float TangentY;
+			float Value;
+
+			FPatchCorner(float InX, float InY, float InTangentX, float InTangentY, float InValue)
+				: X(InX), Y(InY), TangentX(InTangentX), TangentY(InTangentY), Value(InValue)
+			{}
+		};
+
+		/**
+		 * The four corners of the patch, each of which contains and x and y coordinate, and x and y tangent, and a value.
+		 * Indexed in the following order: (X0, Y0) -> (X1, Y0) -> (X1, Y1) -> (X0, Y1)
+		 */
+		FPatchCorner Corners[4];
+	};
+	
+	/** Performs a Coons patch blend on an indexed list of parameters where each parameter has its own set of curves, and outputs all blended parameters */
+	template<typename FocusPointType, typename FocusCurveType>
+	bool IndexedParameterBlend(const TArray<FocusPointType>& FocusPoints, const TArray<FocusCurveType>& FocusCurves, float InFocus, float InZoom, int32 NumParameters, TArray<float>& OutBlendedParameters)
+	{
+		if (FocusPoints.Num() <= 0)
+		{
+			return false;
+		}
+		
+		const LensDataTableUtils::FPointNeighbors PointNeighbors = LensDataTableUtils::FindFocusPoints(InFocus, TConstArrayView<FocusPointType>(FocusPoints));
+		const LensDataTableUtils::FPointNeighbors CurveNeighbors = LensDataTableUtils::FindFocusCurves(InZoom, TConstArrayView<FocusCurveType>(FocusCurves));
+		if (PointNeighbors.IsSinglePoint())
+		{
+			// We are on a zoom curve, or exactly on a corner. Either way, the value can be evaluated directly from the zoom curve
+			for (int32 Index = 0; Index < NumParameters; ++Index)
+			{
+				const FRichCurve* Curve = FocusPoints[PointNeighbors.PreviousIndex].GetCurveForParameter(Index);
+				if (!Curve)
+				{
+					return false;
+				}
+				
+				OutBlendedParameters.Add(Curve->Eval(InZoom));
+			}
+
+			return true;
+		}
+		
+		if (CurveNeighbors.IsSinglePoint())
+		{
+			// We are on one of the focus curves, so return the value evaluated on the focus curves at the specified focus
+			for (int32 Index = 0; Index < NumParameters; ++Index)
+			{
+				const FRichCurve* Curve = FocusCurves[CurveNeighbors.PreviousIndex].GetCurveForParameter(Index);
+				if (!Curve)
+				{
+					return false;
+				}
+				
+				OutBlendedParameters.Add(Curve->Eval(InFocus));
+			}
+
+			return true;
+		}
+
+		const float X0 = FocusCurves[CurveNeighbors.PreviousIndex].Zoom;
+		const float X1 = FocusCurves[CurveNeighbors.NextIndex].Zoom;
+		const float Y0 = FocusPoints[PointNeighbors.PreviousIndex].Focus;
+		const float Y1 = FocusPoints[PointNeighbors.NextIndex].Focus;
+		
+		for (int32 Index = 0; Index < NumParameters; ++Index)
+		{
+			const FRichCurve* X0Curve = FocusPoints[PointNeighbors.PreviousIndex].GetCurveForParameter(Index);
+			const FRichCurve* X1Curve = FocusPoints[PointNeighbors.NextIndex].GetCurveForParameter(Index);
+			const FRichCurve* Y0Curve = FocusCurves[CurveNeighbors.PreviousIndex].GetCurveForParameter(Index);
+			const FRichCurve* Y1Curve = FocusCurves[CurveNeighbors.NextIndex].GetCurveForParameter(Index);
+
+			if (!X0Curve || !X1Curve || !Y0Curve || !Y1Curve)
+			{
+				return false;
+			}
+			
+			FCoonsPatch CoonsPatch(*X0Curve, *X1Curve, *Y0Curve, *Y1Curve, X0, X1, Y0, Y1);
+			OutBlendedParameters.Add(CoonsPatch.Blend(FVector2D(InZoom, InFocus)));
+		}
+		
+		return true;
+	}
+
+	/**
+	 * Parameters for a distortion map blend that configure what values are computed for the blend and how to retrieve the necessary
+	 * blending values from each point being blended
+	 */
+	template<typename TableType>
+	struct FDistortionMapBlendParams
+	{
+		using FocusPointType = typename TableType::FocusPointType;
+		using FocusCurveType = typename TableType::FocusCurveType;
+
+		DECLARE_DELEGATE_FourParams(FGetDisplacementMaps,
+			const FocusPointType& /* FocusPoint */,
+			const FocusCurveType& /* FocusCurve */,
+			UTextureRenderTarget2D*& /* OutUndistortedMap */,
+			UTextureRenderTarget2D*& /* OutDistortedMap */);
+		
+		DECLARE_DELEGATE_RetVal_TwoParams(TOptional<FDistortionInfo>, FGetDistortionParameters,
+			const FocusPointType& /* FocusPoint */,
+			const FocusCurveType& /* FocusCurve */);
+		
+		DECLARE_DELEGATE_RetVal_FourParams(float, FProcessDisplacementMaps,
+			const FocusPointType& /* FocusPoint */,
+			const FocusCurveType& /* FocusCurve */,
+			UTextureRenderTarget2D* /* InUndistortedMap */,
+			UTextureRenderTarget2D* /* InDistortedMap */);
+
+		/**
+		 * Callback to use to retrieve the distortion parameters for a specified point being blended.
+		 * If unbound, no distortion parameters will be computed in the blending results
+		 */
+		FGetDistortionParameters GetDistortionParameters;
+
+		/**
+		 * Callback to retrieve the displacement map render targets to use for a specified point being blended.
+		 * If unbound, use the provided displacement maps in the parameters
+		 */
+		FGetDisplacementMaps GetDisplacementMaps;
+		
+		/**
+		 * Callback to use to perform any processing on the distortion maps for blending that returns the computed overscan for the maps.
+		 * If unbound, no overscan blending will be computed in the blending results
+		 */
+		FProcessDisplacementMaps ProcessDisplacementMaps;
+		
+		/** Indicates that shader blending parameters should be calculated */
+		bool bGenerateBlendingParams = false;
+		
+		/** The number of distortion parameters to blend */
+		int32 DistortionParamNum = 0;
+
+		/** When supplied, list of render targets to write the undistorted displacement maps to */
+		TArray<UTextureRenderTarget2D*> UndistortedMaps;
+
+		/** When supplied, list of render targets to write the distorted displacement maps to */
+		TArray<UTextureRenderTarget2D*> DistortedMaps;
+		
+	};
+
+	/** Blended results from a distortion map blend. Optionals are set when the blend input parameters indicated those blended values should be computed */
+	struct FDistortionMapBlendResults
+	{
+		/** Indicates that a distortion map blend successfully occured */
+		bool bValid = false;
+
+		/** The shader blending parameters, if they were computed for the blend */
+		TOptional<FDisplacementMapBlendingParams> BlendingParams = TOptional<FDisplacementMapBlendingParams>();
+
+		/** The blended distortion parameters, if they were computed for the blend */
+		TOptional<FDistortionInfo> BlendedDistortionParams = TOptional<FDistortionInfo>();
+
+		/** The undistorted maps for the blending, if they generated */
+		TOptional<TArray<UTextureRenderTarget2D*>> UndistortedMaps = TOptional<TArray<UTextureRenderTarget2D*>>();
+
+		/** The distorted maps for the blending, if they generated */
+		TOptional<TArray<UTextureRenderTarget2D*>> DistortedMaps = TOptional<TArray<UTextureRenderTarget2D*>>();
+		
+		/** The blended overscan, if it was computed for the blend */
+		TOptional<float> BlendedOverscan = TOptional<float>();
+	};
+
+	/** Performs a Coons patch blend on the distortion map parameters for the specified table at the specified focus and zoom */
+	template<typename TableType>
+	FDistortionMapBlendResults DistortionMapBlend(const TableType& Table, float InFocus, float InZoom, const FDistortionMapBlendParams<TableType>& InParams)
+	{
+		FDistortionMapBlendResults Results;
+		
+		using FocusPointType = typename TableType::FocusPointType;
+		using FocusCurveType = typename TableType::FocusCurveType;
+
+		// Function to get the correct tangents of the curves safely, even in edge cases where there isn't a rich curve key defined for the specified point
+		auto SafeGetTangents = [](const FRichCurve& Curve, const LensDataTableUtils::FPointNeighbors& Points, float& OutPrevTangent, float& OutNextTangent)
+		{
+			bool bHasPrevious = Curve.Keys.IsValidIndex(Points.PreviousIndex);
+			bool bHasNext = Curve.Keys.IsValidIndex(Points.NextIndex);
+		
+			if (bHasPrevious && bHasNext)
+			{
+				const FRichCurveKey& Prev = Curve.Keys[Points.PreviousIndex];
+				const FRichCurveKey& Next = Curve.Keys[Points.NextIndex];
+				if (Prev.InterpMode == ERichCurveInterpMode::RCIM_Cubic)
+				{
+					OutPrevTangent = Prev.LeaveTangent;
+					OutNextTangent = Next.ArriveTangent;
+				}
+				else
+				{
+					float Slope = (Next.Value - Prev.Value) / (Next.Time - Prev.Time);
+					OutPrevTangent = Slope;
+					OutNextTangent = Slope;
+				}
+			}
+			else
+			{
+				// If one or both points don't have a key in the curve, the curve between those two points
+				// will necessarily be flat, and as such both tangents are 0
+				OutPrevTangent = 0;
+				OutNextTangent = 0;
+			}
+		};
+		
+		TConstArrayView<FocusPointType> FocusPoints = Table.GetFocusPoints();
+		TConstArrayView<FocusCurveType> FocusCurves = Table.GetFocusCurves();
+		
+		if (FocusPoints.Num() <= 0)
+		{
+			return Results;
+		}
+
+		Results.bValid = true;
+
+		if (InParams.bGenerateBlendingParams)
+		{
+			Results.BlendingParams = FDisplacementMapBlendingParams();
+			Results.BlendingParams->EvalFocus = InFocus;
+			Results.BlendingParams->EvalZoom = InZoom;
+		}
+
+		if (InParams.ProcessDisplacementMaps.IsBound())
+		{
+			// If the callback to retrieve the displacement maps for each blended point is set, use that as a source for the displacement maps
+			// Otherwise, if any displacement maps were passed in the parameters, use them
+			if (InParams.GetDisplacementMaps.IsBound())
+			{
+				Results.UndistortedMaps = TArray<UTextureRenderTarget2D*>();
+				Results.DistortedMaps = TArray<UTextureRenderTarget2D*>();
+
+				Results.UndistortedMaps->AddZeroed(4);
+				Results.DistortedMaps->AddZeroed(4);
+			}
+			else if (InParams.UndistortedMaps.Num() == 4 && InParams.DistortedMaps.Num() == 4)
+			{
+				Results.UndistortedMaps = InParams.UndistortedMaps;
+				Results.DistortedMaps = InParams.DistortedMaps;
+			}
+		}
+		
+		if (InParams.GetDistortionParameters.IsBound())
+		{
+			Results.BlendedDistortionParams = FDistortionInfo();
+			Results.BlendedDistortionParams->Parameters.AddDefaulted(InParams.DistortionParamNum);
+		}
+		
+		const LensDataTableUtils::FPointNeighbors PointNeighbors = LensDataTableUtils::FindFocusPoints(InFocus, FocusPoints);
+		const LensDataTableUtils::FPointNeighbors CurveNeighbors = LensDataTableUtils::FindFocusCurves(InZoom, FocusCurves);
+			
+		const FocusPointType& PrevFocusPoint = FocusPoints[PointNeighbors.PreviousIndex];
+		const FocusPointType& NextFocusPoint = FocusPoints[PointNeighbors.NextIndex];
+		const FocusCurveType& PrevFocusCurve = FocusCurves[CurveNeighbors.PreviousIndex];
+		const FocusCurveType& NextFocusCurve = FocusCurves[CurveNeighbors.NextIndex];
+		
+		if (PointNeighbors.IsSinglePoint() && CurveNeighbors.IsSinglePoint())
+		{
+			// We are on a corner of the blending patch
+			// Simply retrieve the value at that point
+			if (Results.BlendingParams.IsSet())
+			{
+				Results.BlendingParams->BlendType = EDisplacementMapBlendType::OneFocusOneZoom;
+			}
+
+			if (Results.BlendedDistortionParams.IsSet())
+			{
+				TOptional<FDistortionInfo> Parameters = InParams.GetDistortionParameters.Execute(PrevFocusPoint, PrevFocusCurve);
+				if (Parameters.IsSet())
+				{
+					Results.BlendedDistortionParams = Parameters.GetValue();
+				}
+			}
+
+			if (Results.UndistortedMaps.IsSet())
+			{
+				InParams.GetDisplacementMaps.ExecuteIfBound(PrevFocusPoint, PrevFocusCurve, (*Results.UndistortedMaps)[0], (*Results.DistortedMaps)[0]);
+				Results.BlendedOverscan = InParams.ProcessDisplacementMaps.Execute(PrevFocusPoint, PrevFocusCurve, (*Results.UndistortedMaps)[0], (*Results.DistortedMaps)[0]);
+			}
+		}
+		else if (PointNeighbors.IsSinglePoint())
+		{
+			// We are on a zoom curve edge of the blending patch
+			float PrevTangent, NextTangent;
+			SafeGetTangents(PrevFocusPoint.MapBlendingCurve, CurveNeighbors, PrevTangent, NextTangent);
+			
+			FTangentBezierCurve BlendCurve(PrevFocusCurve.Zoom, NextFocusCurve.Zoom, 0.0f, 0.0f, PrevTangent, NextTangent);
+			
+			if (Results.BlendingParams.IsSet())
+			{
+				Results.BlendingParams->BlendType = EDisplacementMapBlendType::OneFocusTwoZoom;
+				Results.BlendingParams->PatchCorners[0] = FDisplacementMapBlendPatchCorner(PrevFocusCurve.Zoom, PrevFocusPoint.Focus, PrevTangent, 0.0);
+				Results.BlendingParams->PatchCorners[1] = FDisplacementMapBlendPatchCorner(NextFocusCurve.Zoom, PrevFocusPoint.Focus, NextTangent, 0.0);
+			}
+
+			if (Results.BlendedDistortionParams.IsSet())
+			{
+				// Blend all distortion parameters using the map blend curve
+				TOptional<FDistortionInfo> PrevParameters = InParams.GetDistortionParameters.Execute(PrevFocusPoint, PrevFocusCurve);
+				TOptional<FDistortionInfo> NextParameters = InParams.GetDistortionParameters.Execute(PrevFocusPoint, NextFocusCurve);
+				
+				const int32 NumParams = Results.BlendedDistortionParams->Parameters.Num();
+				for (int Index = 0; Index < NumParams; ++Index)
+				{
+					BlendCurve.Y0 = PrevParameters.IsSet() ? PrevParameters.GetValue().Parameters[Index] : 0.0;
+					BlendCurve.Y1 = NextParameters.IsSet() ? NextParameters.GetValue().Parameters[Index] : 0.0;
+
+					Results.BlendedDistortionParams->Parameters[Index] = BlendCurve.Eval(InZoom);
+				}
+			}
+			
+			if (Results.UndistortedMaps.IsSet())
+			{
+				InParams.GetDisplacementMaps.ExecuteIfBound(PrevFocusPoint, PrevFocusCurve, (*Results.UndistortedMaps)[0], (*Results.DistortedMaps)[0]);
+				InParams.GetDisplacementMaps.ExecuteIfBound(PrevFocusPoint, NextFocusCurve, (*Results.UndistortedMaps)[1], (*Results.DistortedMaps)[1]);
+				
+				// Blend the overscan factor using the map blend curve
+				BlendCurve.Y0 = InParams.ProcessDisplacementMaps.Execute(PrevFocusPoint, PrevFocusCurve, (*Results.UndistortedMaps)[0], (*Results.DistortedMaps)[0]);
+				BlendCurve.Y1 = InParams.ProcessDisplacementMaps.Execute(PrevFocusPoint, NextFocusCurve, (*Results.UndistortedMaps)[1], (*Results.DistortedMaps)[1]);
+				
+				Results.BlendedOverscan = BlendCurve.Eval(InZoom);
+			}
+		}
+		else if (CurveNeighbors.IsSinglePoint())
+		{
+			// We are on a focus curve edge of the blending patch
+			float PrevTangent, NextTangent;
+			SafeGetTangents(PrevFocusCurve.MapBlendingCurve, PointNeighbors, PrevTangent, NextTangent);
+			
+			FTangentBezierCurve BlendCurve(PrevFocusPoint.Focus, NextFocusPoint.Focus, 0.0f, 0.0f, PrevTangent, NextTangent);
+			
+			if (Results.BlendingParams.IsSet())
+			{
+				Results.BlendingParams->BlendType = EDisplacementMapBlendType::TwoFocusOneZoom;
+				Results.BlendingParams->PatchCorners[0] = FDisplacementMapBlendPatchCorner(PrevFocusCurve.Zoom, PrevFocusPoint.Focus, 0.0, PrevTangent);
+				Results.BlendingParams->PatchCorners[1] = FDisplacementMapBlendPatchCorner(PrevFocusCurve.Zoom, NextFocusPoint.Focus, 0.0, NextTangent);
+			}
+			
+			if (Results.BlendedDistortionParams.IsSet())
+			{
+				// Blend all distortion parameters using the map blend curve
+				TOptional<FDistortionInfo> PrevParameters = InParams.GetDistortionParameters.Execute(PrevFocusPoint, PrevFocusCurve);
+				TOptional<FDistortionInfo> NextParameters = InParams.GetDistortionParameters.Execute(NextFocusPoint, PrevFocusCurve);
+				
+				const int32 NumParams = Results.BlendedDistortionParams->Parameters.Num();
+				for (int Index = 0; Index < NumParams; ++Index)
+				{
+					BlendCurve.Y0 = PrevParameters.IsSet() ? PrevParameters.GetValue().Parameters[Index] : 0.0;
+					BlendCurve.Y1 = NextParameters.IsSet() ? NextParameters.GetValue().Parameters[Index] : 0.0;
+
+					Results.BlendedDistortionParams->Parameters[Index] = BlendCurve.Eval(InFocus);
+				}
+			}
+			
+			if (Results.UndistortedMaps.IsSet())
+			{
+				InParams.GetDisplacementMaps.ExecuteIfBound(PrevFocusPoint, PrevFocusCurve, (*Results.UndistortedMaps)[0], (*Results.DistortedMaps)[0]);
+				InParams.GetDisplacementMaps.ExecuteIfBound(NextFocusPoint, PrevFocusCurve, (*Results.UndistortedMaps)[1], (*Results.DistortedMaps)[1]);
+				
+				// Blend the overscan factor using the map blend curve
+				BlendCurve.Y0 = InParams.ProcessDisplacementMaps.Execute(PrevFocusPoint, PrevFocusCurve, (*Results.UndistortedMaps)[0], (*Results.DistortedMaps)[0]);
+				BlendCurve.Y1 = InParams.ProcessDisplacementMaps.Execute(NextFocusPoint, PrevFocusCurve, (*Results.UndistortedMaps)[1], (*Results.DistortedMaps)[1]);
+				
+				Results.BlendedOverscan = BlendCurve.Eval(InFocus);
+			}
+		}
+		else
+		{
+			// Otherwise, we are somewhere in the middle of the patch, and must do a full Coons patch blend
+			float XTangents[4];
+			SafeGetTangents(PrevFocusPoint.MapBlendingCurve, CurveNeighbors, XTangents[0], XTangents[1]);
+			SafeGetTangents(NextFocusPoint.MapBlendingCurve, CurveNeighbors, XTangents[2], XTangents[3]);
+		
+			float YTangents[4];
+			SafeGetTangents(PrevFocusCurve.MapBlendingCurve, PointNeighbors, YTangents[0], YTangents[3]);
+			SafeGetTangents(NextFocusCurve.MapBlendingCurve, PointNeighbors, YTangents[1], YTangents[2]);
+			
+			FTangentBezierCoonsPatch CoonsPatch(PrevFocusCurve.Zoom, NextFocusCurve.Zoom, PrevFocusPoint.Focus, NextFocusPoint.Focus, XTangents, YTangents);
+
+			if (Results.BlendingParams.IsSet())
+			{
+				Results.BlendingParams->BlendType = EDisplacementMapBlendType::TwoFocusTwoZoom;
+				Results.BlendingParams->PatchCorners[0] = FDisplacementMapBlendPatchCorner(PrevFocusCurve.Zoom, PrevFocusPoint.Focus, XTangents[0], YTangents[0]);
+				Results.BlendingParams->PatchCorners[1] = FDisplacementMapBlendPatchCorner(NextFocusCurve.Zoom, PrevFocusPoint.Focus, XTangents[1], YTangents[1]);
+				Results.BlendingParams->PatchCorners[2] = FDisplacementMapBlendPatchCorner(NextFocusCurve.Zoom, NextFocusPoint.Focus, XTangents[2], YTangents[2]);
+				Results.BlendingParams->PatchCorners[3] = FDisplacementMapBlendPatchCorner(PrevFocusCurve.Zoom, NextFocusPoint.Focus, XTangents[3], YTangents[3]);
+			}
+			
+			if (Results.BlendedDistortionParams.IsSet())
+			{
+				TOptional<FDistortionInfo> Points[4] =
+				{
+					InParams.GetDistortionParameters.Execute(PrevFocusPoint, PrevFocusCurve),
+					InParams.GetDistortionParameters.Execute(PrevFocusPoint, NextFocusCurve),
+					InParams.GetDistortionParameters.Execute(NextFocusPoint, NextFocusCurve),
+					InParams.GetDistortionParameters.Execute(NextFocusPoint, PrevFocusCurve)
+				};
+
+				// Handle edge case for when a corner point of the patch is not defined. If such a case happens, we
+				// average its value from its two neighbor points, which are guaranteed to exist
+				for (int32 Index = 0; Index < 4; ++Index)
+				{
+					if (!Points[Index].IsSet())
+					{
+						int32 PrevIndex = Index > 0 ? Index - 1 : 3;
+						int32 NextIndex = Index < 3 ? Index + 1 : 0;
+
+						Points[Index] = FDistortionInfo();
+						
+						int32 NumParams = Results.BlendedDistortionParams->Parameters.Num();
+						for (int ParamIndex = 0; ParamIndex < NumParams; ++ParamIndex)
+						{
+							Points[Index]->Parameters.Add(0.5 * (Points[PrevIndex]->Parameters[ParamIndex] + Points[NextIndex]->Parameters[ParamIndex]));
+						}
+					}
+				}
+				
+				const int32 NumParams = Results.BlendedDistortionParams->Parameters.Num();
+				for (int Index = 0; Index < NumParams; ++Index)
+				{
+					CoonsPatch.Corners[0].Value = Points[0]->Parameters[Index];
+					CoonsPatch.Corners[1].Value = Points[1]->Parameters[Index];
+					CoonsPatch.Corners[2].Value = Points[2]->Parameters[Index];
+					CoonsPatch.Corners[3].Value = Points[3]->Parameters[Index];
+				
+					Results.BlendedDistortionParams->Parameters[Index] = CoonsPatch.Blend(FVector2D(InZoom, InFocus));
+				}
+			}
+			
+			if (Results.UndistortedMaps.IsSet())
+			{
+				InParams.GetDisplacementMaps.ExecuteIfBound(PrevFocusPoint, PrevFocusCurve, (*Results.UndistortedMaps)[0], (*Results.DistortedMaps)[0]);
+				InParams.GetDisplacementMaps.ExecuteIfBound(PrevFocusPoint, NextFocusCurve, (*Results.UndistortedMaps)[1], (*Results.DistortedMaps)[1]);
+				InParams.GetDisplacementMaps.ExecuteIfBound(NextFocusPoint, NextFocusCurve, (*Results.UndistortedMaps)[2], (*Results.DistortedMaps)[2]);
+				InParams.GetDisplacementMaps.ExecuteIfBound(NextFocusPoint, PrevFocusCurve, (*Results.UndistortedMaps)[3], (*Results.DistortedMaps)[3]);
+				
+				CoonsPatch.Corners[0].Value = InParams.ProcessDisplacementMaps.Execute(PrevFocusPoint, PrevFocusCurve, (*Results.UndistortedMaps)[0], (*Results.DistortedMaps)[0]);
+				CoonsPatch.Corners[1].Value = InParams.ProcessDisplacementMaps.Execute(PrevFocusPoint, NextFocusCurve, (*Results.UndistortedMaps)[1], (*Results.DistortedMaps)[1]);
+				CoonsPatch.Corners[2].Value = InParams.ProcessDisplacementMaps.Execute(NextFocusPoint, NextFocusCurve, (*Results.UndistortedMaps)[2], (*Results.DistortedMaps)[2]);
+				CoonsPatch.Corners[3].Value = InParams.ProcessDisplacementMaps.Execute(NextFocusPoint, PrevFocusCurve, (*Results.UndistortedMaps)[3], (*Results.DistortedMaps)[3]);
+
+				Results.BlendedOverscan = CoonsPatch.Blend(FVector2D(InZoom, InFocus));
+			}
+		}
+
+		return Results;
+	}
 };
diff --git a/Engine/Plugins/VirtualProduction/CameraCalibrationCore/Source/CameraCalibrationCore/Private/Tables/BaseLensTable.cpp b/Engine/Plugins/VirtualProduction/CameraCalibrationCore/Source/CameraCalibrationCore/Private/Tables/BaseLensTable.cpp
--- a/Engine/Plugins/VirtualProduction/CameraCalibrationCore/Source/CameraCalibrationCore/Private/Tables/BaseLensTable.cpp
+++ b/Engine/Plugins/VirtualProduction/CameraCalibrationCore/Source/CameraCalibrationCore/Private/Tables/BaseLensTable.cpp
@@ -3,6 +3,98 @@
 #include "Tables/BaseLensTable.h"
 #include "LensFile.h"
 
+FKeyHandle FBaseFocusCurve::AddPointToCurve(FRichCurve& InCurve, float InFocus, float InValue, float InputTolerance, FKeyHandle InOptionalKeyHandle)
+{
+	const FKeyHandle ExistingKeyHandle = SetPointInCurve(InCurve, InFocus, InValue, InputTolerance);
+	if (ExistingKeyHandle != FKeyHandle::Invalid())
+	{
+		return ExistingKeyHandle;
+	}
+	
+	const FKeyHandle NewKeyHandle = InCurve.AddKey(InFocus, InValue, false, InOptionalKeyHandle);
+	InCurve.SetKeyTangentMode(NewKeyHandle, ERichCurveTangentMode::RCTM_None);
+	InCurve.SetKeyInterpMode(NewKeyHandle, RCIM_Linear);
+
+	return NewKeyHandle;
+}
+
+FKeyHandle FBaseFocusCurve::SetPointInCurve(FRichCurve& InCurve, float InFocus, float InValue, float InputTolerance)
+{
+	const FKeyHandle KeyHandle = InCurve.FindKey(InFocus, InputTolerance);
+	if (KeyHandle != FKeyHandle::Invalid())
+	{
+		InCurve.SetKeyValue(KeyHandle, InValue);
+	}
+
+	return KeyHandle;
+}
+
+void FBaseFocusCurve::DeletePointFromCurve(FRichCurve& InCurve, float InFocus, float InputTolerance)
+{
+	const FKeyHandle KeyHandle = InCurve.FindKey(InFocus, InputTolerance);
+	if (KeyHandle != FKeyHandle::Invalid())
+	{
+		InCurve.DeleteKey(KeyHandle);
+	}
+}
+
+void FBaseFocusCurve::ChangeFocusInCurve(FRichCurve& InCurve, float InExistingFocus, float InNewFocus, float InputTolerance)
+{
+	const FKeyHandle KeyHandle = InCurve.FindKey(InExistingFocus, InputTolerance);
+	if (KeyHandle != FKeyHandle::Invalid())
+	{
+		InCurve.SetKeyTime(KeyHandle, InNewFocus);
+	}
+}
+
+void FBaseFocusCurve::MergeFocusInCurve(FRichCurve& InCurve, float InExistingFocus, float InNewFocus, bool bReplaceExisting, float InputTolerance)
+{
+	const FKeyHandle KeyHandle = InCurve.FindKey(InExistingFocus, InputTolerance);
+	if (KeyHandle != FKeyHandle::Invalid())
+	{
+		const FKeyHandle CxNewHandle = InCurve.FindKey(InNewFocus, InputTolerance);
+		if (CxNewHandle != FKeyHandle::Invalid() && bReplaceExisting)
+		{
+			InCurve.Keys[InCurve.GetIndexSafe(CxNewHandle)] = InCurve.GetKey(KeyHandle);
+			InCurve.DeleteKey(KeyHandle);
+		}
+		else if (CxNewHandle == FKeyHandle::Invalid())
+		{
+			InCurve.SetKeyTime(KeyHandle, InNewFocus);	
+		}
+	}
+}
+
+void FBaseLensTable::CopyCurveKeys(const FRichCurve& InSourceCurve, FRichCurve& InDestCurve, TArrayView<const FKeyHandle> InKeys)
+{
+	for (int32 Index = 0; Index < InKeys.Num(); ++Index)
+	{
+		const FKeyHandle Handle = InKeys[Index];
+		const int32 KeyIndex = InSourceCurve.GetIndexSafe(Handle);
+		if (KeyIndex != INDEX_NONE)
+		{
+			InDestCurve.Keys[KeyIndex] = InSourceCurve.GetKey(Handle);
+		}
+	}
+
+	InDestCurve.AutoSetTangents();
+}
+
+void FBaseLensTable::PropagateCurveValuesToCrossCurves(const FRichCurve& InCurve, float InCrossCurveTime, TFunctionRef<FRichCurve*(float)> GetCurveFn)
+{
+	for (auto Iter(InCurve.GetKeyIterator()); Iter; ++Iter)
+	{
+		if (FRichCurve* CrossCurve = GetCurveFn(Iter->Time))
+		{
+			const FKeyHandle KeyHandle = CrossCurve->FindKey(InCrossCurveTime);
+			if (KeyHandle != FKeyHandle::Invalid())
+			{
+				CrossCurve->SetKeyValue(KeyHandle, Iter->Value);
+			}
+		}
+	}
+}
+
 FName FBaseLensTable::GetFriendlyPointName(ELensDataCategory InCategory)
 {
 	switch (InCategory)
diff --git a/Engine/Plugins/VirtualProduction/CameraCalibrationCore/Source/CameraCalibrationCore/Private/Tables/DistortionParametersTable.cpp b/Engine/Plugins/VirtualProduction/CameraCalibrationCore/Source/CameraCalibrationCore/Private/Tables/DistortionParametersTable.cpp
--- a/Engine/Plugins/VirtualProduction/CameraCalibrationCore/Source/CameraCalibrationCore/Private/Tables/DistortionParametersTable.cpp
+++ b/Engine/Plugins/VirtualProduction/CameraCalibrationCore/Source/CameraCalibrationCore/Private/Tables/DistortionParametersTable.cpp
@@ -91,6 +91,16 @@ bool FDistortionFocusPoint::IsEmpty() const
 	return MapBlendingCurve.IsEmpty();
 }
 
+const FRichCurve* FDistortionFocusPoint::GetCurveForParameter(int32 InParameterIndex) const
+{
+	if (InParameterIndex != FDistortionTable::FParameters::Aggregate)
+	{
+		return nullptr;
+	}
+
+	return &MapBlendingCurve;
+}
+
 void FDistortionFocusPoint::SetParameterValue(int32 InZoomIndex, float InZoomValue, int32 InParameterIndex, float InParameterValue)
 {
 	//We can't move keys on the time axis so our indices should match
@@ -103,6 +113,46 @@ void FDistortionFocusPoint::SetParameterValue(int32 InZoomIndex, float InZoomVal
 	}
 }
 
+void FDistortionFocusCurve::AddPoint(float InFocus, const FDistortionInfo& InData, float InputTolerance)
+{
+	AddPointToCurve(MapBlendingCurve, InFocus, InFocus, InputTolerance);
+}
+
+void FDistortionFocusCurve::SetPoint(float InFocus, const FDistortionInfo& InData, float InputTolerance)
+{
+	// No need to update the curve since curve time = curve value
+}
+
+void FDistortionFocusCurve::RemovePoint(float InFocus, float InputTolerance)
+{
+	DeletePointFromCurve(MapBlendingCurve, InFocus, InputTolerance);
+}
+
+void FDistortionFocusCurve::ChangeFocus(float InExistingFocus, float InNewFocus, float InputTolerance)
+{
+	ChangeFocusInCurve(MapBlendingCurve, InExistingFocus, InNewFocus, InputTolerance);
+}
+
+void FDistortionFocusCurve::MergeFocus(float InExistingFocus, float InNewFocus, bool bReplaceExisting, float InputTolerance)
+{
+	MergeFocusInCurve(MapBlendingCurve, InExistingFocus, InNewFocus, bReplaceExisting, InputTolerance);
+}
+
+bool FDistortionFocusCurve::IsEmpty() const
+{
+	return !MapBlendingCurve.GetNumKeys();
+}
+
+const FRichCurve* FDistortionFocusCurve::GetCurveForParameter(int32 InParameterIndex) const
+{
+	if (InParameterIndex != FDistortionTable::FParameters::Aggregate)
+	{
+		return nullptr;
+	}
+
+	return &MapBlendingCurve;
+}
+
 int32 FDistortionTable::GetTotalPointNum() const
 {
 	return LensDataTableUtils::GetTotalPointNum(FocusPoints);
@@ -140,22 +190,31 @@ bool FDistortionTable::BuildParameterCurveAtFocus(float InFocus, int32 Parameter
 
 bool FDistortionTable::BuildParameterCurveAtZoom(float InZoom, int32 InParameterIndex, FRichCurve& OutCurve) const
 {
+	if (InParameterIndex == FParameters::Aggregate)
+	{
+		const FDistortionFocusCurve* Curve = GetFocusCurve(InZoom);
+		if (!Curve)
+		{
+			return false;
+		}
+
+		OutCurve = Curve->MapBlendingCurve;
+		return true;
+	}
+	
 	for (const FDistortionFocusPoint& FocusPoint : FocusPoints)
 	{
 		FDistortionInfo ZoomPoint;
 		if (FocusPoint.GetPoint(InZoom, ZoomPoint))
 		{
 			float Value = 0.0;
-			if (ZoomPoint.Parameters.IsValidIndex(InParameterIndex))
+			if (!ZoomPoint.Parameters.IsValidIndex(InParameterIndex))
 			{
-				Value = ZoomPoint.Parameters[InParameterIndex];
-			}
-			else
-			{
-				FKeyHandle BlendPointKey = FocusPoint.MapBlendingCurve.FindKey(InZoom);
-				Value = FocusPoint.MapBlendingCurve.GetKeyValue(BlendPointKey);
+				return false;
 			}
 			
+			Value = ZoomPoint.Parameters[InParameterIndex];
+			
 			const FKeyHandle NewKeyHandle = OutCurve.AddKey(FocusPoint.Focus, Value);
 			FRichCurveKey& NewKey = OutCurve.GetKey(NewKeyHandle);
 			NewKey.TangentMode = ERichCurveTangentMode::RCTM_None;
@@ -194,34 +253,40 @@ void FDistortionTable::SetParameterCurveKeysAtFocus(float InFocus, int32 InParam
 
 void FDistortionTable::SetParameterCurveKeysAtZoom(float InZoom, int32 InParameterIndex, const FRichCurve& InSourceCurve, TArrayView<const FKeyHandle> InKeys)
 {
+	// The aggregate curve's attributes can be changed, but not its key values, so simply copy the key attributes from the source curve
+	// to the corresponding focus curve
+	if (InParameterIndex == FParameters::Aggregate)
+	{
+		FDistortionFocusCurve* Curve = GetFocusCurve(InZoom);
+		if (!Curve)
+		{
+			return;
+		}
+
+		CopyCurveKeys(InSourceCurve, Curve->MapBlendingCurve, InKeys);
+		return;
+	}
+
+	// For every other parameter, only the curve's key values can be changed, so iterate over all keys to find the
+	// appropriate focus/zoom point and update its value to match the curve's key values
 	for (const FKeyHandle& KeyHandle : InKeys)
 	{
-		// Assume the focus keys are put into the source curve in the same order as they are stored internally
-		const int32 KeyIndex = InSourceCurve.GetIndexSafe(KeyHandle);
-		if (KeyIndex != INDEX_NONE)
+		const FRichCurveKey& Key = InSourceCurve.GetKey(KeyHandle);
+		if (FDistortionFocusPoint* FocusPoint = GetFocusPoint(Key.Time))
 		{
-			if (ensure(FocusPoints.IsValidIndex(KeyIndex)))
+			FDistortionInfo ZoomPoint;
+			if (!FocusPoint->GetPoint(InZoom, ZoomPoint))
 			{
-				//We can't move keys on the time axis so our indices should match
-				const FRichCurveKey& Key = InSourceCurve.GetKey(KeyHandle);
-				FDistortionFocusPoint& FocusPoint = FocusPoints[KeyIndex];
-				
-				if (InParameterIndex == FParameters::Aggregate)
-				{
-					FocusPoint.MapBlendingCurve.Keys[KeyIndex] = Key;
-				}
-				else
-				{
-					FDistortionInfo ZoomPoint;
-					if (!FocusPoint.GetPoint(InZoom, ZoomPoint))
-					{
-						continue;
-					}
-					
-					ZoomPoint.Parameters[InParameterIndex] = Key.Value;
-					FocusPoint.SetPoint(InZoom, ZoomPoint);
-				}
+				continue;
 			}
+
+			if (!ZoomPoint.Parameters.IsValidIndex(InParameterIndex))
+			{
+				return;
+			}
+		
+			ZoomPoint.Parameters[InParameterIndex] = Key.Value;
+			FocusPoint->SetPoint(InZoom, ZoomPoint);
 		}
 	}
 }
@@ -256,6 +321,16 @@ FDistortionFocusPoint* FDistortionTable::GetFocusPoint(float InFocus, float Inpu
 	return FocusPoints.FindByPredicate([InFocus, InputTolerance](const FDistortionFocusPoint& Point) { return FMath::IsNearlyEqual(Point.Focus, InFocus, InputTolerance); });
 }
 
+const FDistortionFocusCurve* FDistortionTable::GetFocusCurve(float InZoom, float InputTolerance) const
+{
+	return FocusCurves.FindByPredicate([InZoom, InputTolerance](const FDistortionFocusCurve& Curve) { return FMath::IsNearlyEqual(Curve.Zoom, InZoom, InputTolerance); });
+}
+
+FDistortionFocusCurve* FDistortionTable::GetFocusCurve(float InZoom, float InputTolerance)
+{
+	return FocusCurves.FindByPredicate([InZoom, InputTolerance](const FDistortionFocusCurve& Curve) { return FMath::IsNearlyEqual(Curve.Zoom, InZoom, InputTolerance); });
+}
+
 void FDistortionTable::ForEachPoint(FFocusPointCallback InCallback) const
 {
 	for (const FDistortionFocusPoint& Point : FocusPoints)
@@ -274,9 +349,20 @@ TArray<FDistortionFocusPoint>& FDistortionTable::GetFocusPoints()
 	return FocusPoints;
 }
 
+TConstArrayView<FDistortionFocusCurve> FDistortionTable::GetFocusCurves() const
+{
+	return FocusCurves;
+}
+
+TArray<FDistortionFocusCurve>& FDistortionTable::GetFocusCurves()
+{
+	return FocusCurves;
+}
+
 void FDistortionTable::RemoveFocusPoint(float InFocus)
 {
 	LensDataTableUtils::RemoveFocusPoint(FocusPoints, InFocus);
+	LensDataTableUtils::RemoveFocusFromFocusCurves(FocusCurves, InFocus);
 }
 
 bool FDistortionTable::HasFocusPoint(float InFocus, float InputTolerance) const
@@ -287,16 +373,19 @@ bool FDistortionTable::HasFocusPoint(float InFocus, float InputTolerance) const
 void FDistortionTable::ChangeFocusPoint(float InExistingFocus, float InNewFocus, float InputTolerance)
 {
 	LensDataTableUtils::ChangeFocusPoint(FocusPoints, InExistingFocus, InNewFocus, InputTolerance);
+	LensDataTableUtils::ChangeFocusInFocusCurves(FocusCurves, InExistingFocus, InNewFocus, InputTolerance);
 }
 
 void FDistortionTable::MergeFocusPoint(float InSrcFocus, float InDestFocus, bool bReplaceExistingZoomPoints, float InputTolerance)
 {
 	LensDataTableUtils::MergeFocusPoint(FocusPoints, InSrcFocus, InDestFocus, bReplaceExistingZoomPoints, InputTolerance);
+	LensDataTableUtils::MergeFocusInFocusCurves(FocusCurves, InSrcFocus, InDestFocus, bReplaceExistingZoomPoints, InputTolerance);
 }
 
 void FDistortionTable::RemoveZoomPoint(float InFocus, float InZoom)
 {
 	LensDataTableUtils::RemoveZoomPoint(FocusPoints, InFocus, InZoom);
+	LensDataTableUtils::RemoveZoomFromFocusCurves(FocusCurves, InFocus, InZoom);
 }
 
 bool FDistortionTable::HasZoomPoint(float InFocus, float InZoom, float InputTolerance)
@@ -307,6 +396,14 @@ bool FDistortionTable::HasZoomPoint(float InFocus, float InZoom, float InputTole
 void FDistortionTable::ChangeZoomPoint(float InFocus, float InExistingZoom, float InNewZoom, float InputTolerance)
 {
 	LensDataTableUtils::ChangeZoomPoint(FocusPoints, InFocus, InExistingZoom, InNewZoom, InputTolerance);
+
+	FDistortionInfo Data;
+	if (!GetPoint(InFocus, InNewZoom, Data, InputTolerance))
+	{
+		return;
+	}
+
+	LensDataTableUtils::ChangeZoomInFocusCurves(FocusCurves, InFocus, InExistingZoom, InNewZoom, Data, InputTolerance);
 }
 
 bool FDistortionTable::DoesFocusPointExists(float InFocus, float InputTolerance) const
@@ -353,7 +450,13 @@ TMap<ELensDataCategory, FLinkPointMetadata> FDistortionTable::GetLinkedCategorie
 
 bool FDistortionTable::AddPoint(float InFocus, float InZoom, const FDistortionInfo& InData, float InputTolerance, bool bIsCalibrationPoint)
 {
-	return LensDataTableUtils::AddPoint(FocusPoints, InFocus, InZoom, InData, InputTolerance, bIsCalibrationPoint);
+	if (!LensDataTableUtils::AddPoint(FocusPoints, InFocus, InZoom, InData, InputTolerance, bIsCalibrationPoint))
+	{
+		return false;
+	}
+
+	LensDataTableUtils::AddPointToFocusCurve(FocusCurves, InFocus, InZoom, InData, InputTolerance);
+	return true;
 }
 
 bool FDistortionTable::GetPoint(const float InFocus, const float InZoom, FDistortionInfo& OutData, float InputTolerance) const
@@ -374,7 +477,20 @@ bool FDistortionTable::GetPoint(const float InFocus, const float InZoom, FDistor
 
 bool FDistortionTable::SetPoint(float InFocus, float InZoom, const FDistortionInfo& InData, float InputTolerance)
 {
-	return LensDataTableUtils::SetPoint(*this, InFocus, InZoom, InData, InputTolerance);
+	if (!LensDataTableUtils::SetPoint(*this, InFocus, InZoom, InData, InputTolerance))
+	{
+		return false;
+	}
+
+	LensDataTableUtils::SetPointInFocusCurve(FocusCurves, InFocus, InZoom, InData, InputTolerance);
+	return true;
+}
+
+void FDistortionTable::BuildFocusCurves()
+{
+	// Ensure that the focus curves are empty before building them from the table data
+	FocusCurves.Empty();
+	LensDataTableUtils::BuildFocusCurves(FocusPoints, FocusCurves);
 }
 
 
diff --git a/Engine/Plugins/VirtualProduction/CameraCalibrationCore/Source/CameraCalibrationCore/Private/Tables/FocalLengthTable.cpp b/Engine/Plugins/VirtualProduction/CameraCalibrationCore/Source/CameraCalibrationCore/Private/Tables/FocalLengthTable.cpp
--- a/Engine/Plugins/VirtualProduction/CameraCalibrationCore/Source/CameraCalibrationCore/Private/Tables/FocalLengthTable.cpp
+++ b/Engine/Plugins/VirtualProduction/CameraCalibrationCore/Source/CameraCalibrationCore/Private/Tables/FocalLengthTable.cpp
@@ -131,6 +131,87 @@ bool FFocalLengthFocusPoint::IsEmpty() const
 	return Fx.IsEmpty();
 }
 
+const FRichCurve* FFocalLengthFocusPoint::GetCurveForParameter(int32 InParameterIndex) const
+{
+	if (FFocalLengthTable::FParameters::IsValid(InParameterIndex))
+	{
+		if (InParameterIndex == FFocalLengthTable::FParameters::Fx)
+		{
+			return &Fx;
+		}
+
+		if (InParameterIndex == FFocalLengthTable::FParameters::Fy)
+		{
+			return &Fy;
+		}
+	}
+
+	return nullptr;
+}
+
+FRichCurve* FFocalLengthFocusPoint::GetCurveForParameter(int32 InParameterIndex)
+{
+	return const_cast<FRichCurve*>(const_cast<const FFocalLengthFocusPoint*>(this)->GetCurveForParameter(InParameterIndex));
+}
+
+void FFocalLengthFocusCurve::AddPoint(float InFocus, const FFocalLengthInfo& InData, float InputTolerance)
+{
+	const FKeyHandle KeyHandle = AddPointToCurve(Fx, InFocus, InData.FxFy.X, InputTolerance);
+	AddPointToCurve(Fy, InFocus, InData.FxFy.Y, InputTolerance, KeyHandle);
+}
+
+void FFocalLengthFocusCurve::SetPoint(float InFocus, const FFocalLengthInfo& InData, float InputTolerance)
+{
+	SetPointInCurve(Fx, InFocus, InData.FxFy.X, InputTolerance);
+	SetPointInCurve(Fy, InFocus, InData.FxFy.Y, InputTolerance);
+}
+
+void FFocalLengthFocusCurve::RemovePoint(float InFocus, float InputTolerance)
+{
+	DeletePointFromCurve(Fx, InFocus, InputTolerance);
+	DeletePointFromCurve(Fy, InFocus, InputTolerance);
+}
+
+void FFocalLengthFocusCurve::ChangeFocus(float InExistingFocus, float InNewFocus, float InputTolerance)
+{
+	ChangeFocusInCurve(Fx, InExistingFocus, InNewFocus, InputTolerance);
+	ChangeFocusInCurve(Fy, InExistingFocus, InNewFocus, InputTolerance);
+}
+
+void FFocalLengthFocusCurve::MergeFocus(float InExistingFocus, float InNewFocus, bool bReplaceExisting, float InputTolerance)
+{
+	MergeFocusInCurve(Fx, InExistingFocus, InNewFocus, bReplaceExisting, InputTolerance);
+	MergeFocusInCurve(Fy, InExistingFocus, InNewFocus, bReplaceExisting, InputTolerance);
+}
+
+bool FFocalLengthFocusCurve::IsEmpty() const
+{
+	return !Fx.GetNumKeys() && !Fy.GetNumKeys();
+}
+
+const FRichCurve* FFocalLengthFocusCurve::GetCurveForParameter(int32 InParameterIndex) const
+{
+	if (FFocalLengthTable::FParameters::IsValid(InParameterIndex))
+	{
+		if (InParameterIndex == FFocalLengthTable::FParameters::Fx)
+		{
+			return &Fx;
+		}
+
+		if (InParameterIndex == FFocalLengthTable::FParameters::Fy)
+		{
+			return &Fy;
+		}
+	}
+
+	return nullptr;
+}
+
+FRichCurve* FFocalLengthFocusCurve::GetCurveForParameter(int32 InParameterIndex)
+{
+	return const_cast<FRichCurve*>(const_cast<const FFocalLengthFocusCurve*>(this)->GetCurveForParameter(InParameterIndex));
+}
+
 bool FFocalLengthTable::DoesZoomPointExists(float InFocus, float InZoom, float InputTolerance) const
 {
 	FFocalLengthInfo FocalLengthInfo;
@@ -225,36 +306,59 @@ bool FFocalLengthTable::BuildParameterCurveAtFocus(float InFocus, int32 Paramete
 
 bool FFocalLengthTable::BuildParameterCurveAtZoom(float InZoom, int32 InParameterIndex, FRichCurve& OutCurve) const
 {
-	for (const FFocalLengthFocusPoint& FocusPoint : FocusPoints)
+	if (!FParameters::IsValidOrAggregate(InParameterIndex))
+	{
+		return false;
+	}
+
+	const FFocalLengthFocusCurve* FocusCurve = GetFocusCurve(InZoom);
+	if (!FocusCurve)
+	{
+		return false;
+	}
+
+	if (InParameterIndex == FParameters::Aggregate)
 	{
-		FFocalLengthInfo ZoomPoint;
-		if (FocusPoint.GetPoint(InZoom, ZoomPoint))
+		// The aggregate curve is just the x curve scaled by the lens file's sensor
+		const FRichCurve& ActiveCurve = FocusCurve->Fx;
+		auto HandleIter = ActiveCurve.GetKeyHandleIterator();
+		for (const FRichCurveKey& Key : ActiveCurve.GetConstRefOfKeys())
 		{
-			float Value = ZoomPoint.FxFy.X;
-			if (FParameters::IsValid(InParameterIndex))
-			{
-				Value = ZoomPoint.FxFy[InParameterIndex];
-			}
-			else
-			{
-				if (ULensFile* LensFilePtr = GetLensFile())
-				{
-					Value = ZoomPoint.FxFy.X * LensFilePtr->LensInfo.SensorDimensions.X;
-				}
-			}
-			
-			const FKeyHandle NewKeyHandle = OutCurve.AddKey(FocusPoint.Focus, Value);
-			FRichCurveKey& NewKey = OutCurve.GetKey(NewKeyHandle);
-			NewKey.TangentMode = ERichCurveTangentMode::RCTM_None;
-			NewKey.InterpMode = ERichCurveInterpMode::RCIM_Linear;
+			ULensFile* LensFilePtr = GetLensFile();
+			const float Scale = LensFilePtr ? LensFilePtr->LensInfo.SensorDimensions.X : 1.0f;
+			OutCurve.AddKey(Key.Time, Key.Value * Scale, false, *HandleIter);
+
+			FRichCurveKey& NewKey = OutCurve.GetKey(*HandleIter);
+			NewKey.TangentMode = Key.TangentMode;
+			NewKey.InterpMode = Key.InterpMode;
+			NewKey.ArriveTangent = Key.ArriveTangent * Scale;
+			NewKey.LeaveTangent = Key.LeaveTangent * Scale;
+			++HandleIter;
 		}
+
+		return true;
+	}
+	else if (InParameterIndex == FParameters::Fx)
+	{
+		OutCurve = FocusCurve->Fx;
+		return true;
+	}
+	else if (InParameterIndex == FParameters::Fy)
+	{
+		OutCurve = FocusCurve->Fy;
+		return true;
 	}
 
-	return true;
+	return false;
 }
 
 void FFocalLengthTable::SetParameterCurveKeysAtFocus(float InFocus, int32 InParameterIndex, const FRichCurve& InSourceCurve, TArrayView<const FKeyHandle> InKeys)
 {
+	if (!FParameters::IsValidOrAggregate(InParameterIndex))
+	{
+		return;
+	}
+	
 	FFocalLengthFocusPoint* FocusPoint = GetFocusPoint(InFocus);
 	if (!FocusPoint)
 	{
@@ -275,13 +379,9 @@ void FFocalLengthTable::SetParameterCurveKeysAtFocus(float InFocus, int32 InPara
 		
 		FxFyIndex = 0; //mm focal length curve changes Fx
 	}
-	else if (InParameterIndex == FParameters::Fx)
-	{
-		ActiveCurve = &FocusPoint->Fx;
-	}
-	else if (InParameterIndex == FParameters::Fy)
+	else
 	{
-		ActiveCurve = &FocusPoint->Fy;
+		ActiveCurve = FocusPoint->GetCurveForParameter(InParameterIndex);
 	}
 
 	if (!ActiveCurve)
@@ -311,6 +411,16 @@ void FFocalLengthTable::SetParameterCurveKeysAtFocus(float InFocus, int32 InPara
 	}	
 
 	ActiveCurve->AutoSetTangents();
+	
+	PropagateCurveValuesToCrossCurves(*ActiveCurve, InFocus, [this, InParameterIndex](float InZoom)->FRichCurve*
+	{
+		if (FFocalLengthFocusCurve* Curve = GetFocusCurve(InZoom))
+		{
+			return Curve->GetCurveForParameter(InParameterIndex);
+		}
+
+		return nullptr;
+	});
 }
 
 void FFocalLengthTable::SetParameterCurveKeysAtZoom(float InZoom, int32 InParameterIndex, const FRichCurve& InSourceCurve, TArrayView<const FKeyHandle> InKeys)
@@ -320,47 +430,75 @@ void FFocalLengthTable::SetParameterCurveKeysAtZoom(float InZoom, int32 InParame
 		return;
 	}
 	
+	FFocalLengthFocusCurve* FocusCurve = GetFocusCurve(InZoom);
+	if (!FocusCurve)
+	{
+		return;
+	}
+	
+	FRichCurve* ActiveCurve = nullptr;
+	float Scale = 1.0f;
+	int32 FxFyIndex = InParameterIndex;
+	if (InParameterIndex == FParameters::Aggregate)
+	{
+		ActiveCurve = &FocusCurve->Fx;
+		
+		if (ULensFile* LensFilePtr = GetLensFile())
+		{
+			Scale = 1.0f / LensFilePtr->LensInfo.SensorDimensions.X;
+		}
+		
+		FxFyIndex = 0; //mm focal length curve changes Fx
+	}
+	else
+	{
+		ActiveCurve = FocusCurve->GetCurveForParameter(InParameterIndex);
+	}
+
+	if (!ActiveCurve)
+	{
+		return;
+	}
+
 	for (const FKeyHandle& KeyHandle : InKeys)
 	{
-		// Assume the focus keys are put into the source curve in the same order as they are stored internally
 		const int32 KeyIndex = InSourceCurve.GetIndexSafe(KeyHandle);
 		if (KeyIndex != INDEX_NONE)
 		{
-			if (ensure(FocusPoints.IsValidIndex(KeyIndex)))
+			if(ensure(ActiveCurve->Keys.IsValidIndex(KeyIndex)))
 			{
-				FFocalLengthFocusPoint& FocusPoint = FocusPoints[KeyIndex];
-				FFocalLengthInfo ZoomPoint;
-				if (!FocusPoint.GetPoint(InZoom, ZoomPoint))
-				{
-					continue;
-				}
+				const FRichCurveKey& SourceKey = InSourceCurve.GetKey(KeyHandle);
+				FRichCurveKey& DestinationKey = ActiveCurve->Keys[KeyIndex];
 				
-				float Scale = 1.0f;
-				int32 FxFyIndex = InParameterIndex;
-				if (InParameterIndex == FParameters::Aggregate)
+				DestinationKey.Value = SourceKey.Value * Scale;
+				DestinationKey.InterpMode = SourceKey.InterpMode;
+				DestinationKey.ArriveTangent = SourceKey.ArriveTangent * Scale;
+				DestinationKey.LeaveTangent = SourceKey.LeaveTangent * Scale;
+				DestinationKey.TangentMode = SourceKey.TangentMode;
+
+				if (FFocalLengthFocusPoint* FocusPoint = GetFocusPoint(SourceKey.Time))
 				{
-					if (ULensFile* LensFilePtr = GetLensFile())
+					FFocalLengthInfo ZoomPoint;
+					if (FocusPoint->GetPoint(InZoom, ZoomPoint))
 					{
-						Scale = 1.0f / LensFilePtr->LensInfo.SensorDimensions.X;
+						ZoomPoint.FxFy[FxFyIndex] = InSourceCurve.GetKeyValue(KeyHandle) * Scale;
+						FocusPoint->SetPoint(InZoom, ZoomPoint);
+						
+						if (InParameterIndex == FParameters::Fy)
+						{
+							FocusPoint->Fy.AutoSetTangents();
+						}
+						else
+						{
+							FocusPoint->Fx.AutoSetTangents();
+						}
 					}
-		
-					FxFyIndex = 0; //mm focal length curve changes Fx
-				}
-
-				ZoomPoint.FxFy[FxFyIndex] = InSourceCurve.GetKeyValue(KeyHandle) * Scale;
-				FocusPoint.SetPoint(InZoom, ZoomPoint);
-
-				if (InParameterIndex == FParameters::Fy)
-				{
-					FocusPoint.Fy.AutoSetTangents();
-				}
-				else
-				{
-					FocusPoint.Fx.AutoSetTangents();
 				}
 			}
 		}
-	}
+	}	
+
+	ActiveCurve->AutoSetTangents();
 }
 
 TRange<double> FFocalLengthTable::GetCurveKeyPositionRange(int32 InParameterIndex) const
@@ -419,11 +557,36 @@ FFocalLengthFocusPoint* FFocalLengthTable::GetFocusPoint(float InFocus, float In
 	return FocusPoints.FindByPredicate([InFocus, InputTolerance](const FFocalLengthFocusPoint& Point) { return FMath::IsNearlyEqual(Point.Focus, InFocus, InputTolerance); });
 }
 
+const FFocalLengthFocusCurve* FFocalLengthTable::GetFocusCurve(float InZoom, float InputTolerance) const
+{
+	return FocusCurves.FindByPredicate([InZoom, InputTolerance](const FFocalLengthFocusCurve& Curve) { return FMath::IsNearlyEqual(Curve.Zoom, InZoom, InputTolerance); });
+}
+
+FFocalLengthFocusCurve* FFocalLengthTable::GetFocusCurve(float InZoom, float InputTolerance)
+{
+	return FocusCurves.FindByPredicate([InZoom, InputTolerance](const FFocalLengthFocusCurve& Curve) { return FMath::IsNearlyEqual(Curve.Zoom, InZoom, InputTolerance); });
+}
+
 TConstArrayView<FFocalLengthFocusPoint> FFocalLengthTable::GetFocusPoints() const
 {
 	return FocusPoints;
 }
 
+TArray<FFocalLengthFocusPoint>& FFocalLengthTable::GetFocusPoints()
+{
+	return FocusPoints;
+}
+
+TConstArrayView<FFocalLengthFocusCurve> FFocalLengthTable::GetFocusCurves() const
+{
+	return FocusCurves;
+}
+
+TArray<FFocalLengthFocusCurve>& FFocalLengthTable::GetFocusCurves()
+{
+	return FocusCurves;
+}
+
 void FFocalLengthTable::ForEachPoint(FFocusPointCallback InCallback) const
 {
 	for (const FFocalLengthFocusPoint& Point : FocusPoints)
@@ -435,6 +598,7 @@ void FFocalLengthTable::ForEachPoint(FFocusPointCallback InCallback) const
 void FFocalLengthTable::RemoveFocusPoint(float InFocus)
 {
 	LensDataTableUtils::RemoveFocusPoint(FocusPoints, InFocus);
+	LensDataTableUtils::RemoveFocusFromFocusCurves(FocusCurves, InFocus);
 }
 
 bool FFocalLengthTable::HasFocusPoint(float InFocus,  float InputTolerance) const
@@ -445,16 +609,19 @@ bool FFocalLengthTable::HasFocusPoint(float InFocus,  float InputTolerance) cons
 void FFocalLengthTable::ChangeFocusPoint(float InExistingFocus, float InNewFocus,  float InputTolerance)
 {
 	LensDataTableUtils::ChangeFocusPoint(FocusPoints, InExistingFocus, InNewFocus, InputTolerance);
+	LensDataTableUtils::ChangeFocusInFocusCurves(FocusCurves, InExistingFocus, InNewFocus, InputTolerance);
 }
 
 void FFocalLengthTable::MergeFocusPoint(float InSrcFocus, float InDestFocus, bool bReplaceExistingZoomPoints,  float InputTolerance)
 {
 	LensDataTableUtils::MergeFocusPoint(FocusPoints, InSrcFocus, InDestFocus, bReplaceExistingZoomPoints, InputTolerance);
+	LensDataTableUtils::MergeFocusInFocusCurves(FocusCurves, InSrcFocus, InDestFocus, bReplaceExistingZoomPoints, InputTolerance);
 }
 
 void FFocalLengthTable::RemoveZoomPoint(float InFocus, float InZoom)
 {
 	LensDataTableUtils::RemoveZoomPoint(FocusPoints, InFocus, InZoom);
+	LensDataTableUtils::RemoveZoomFromFocusCurves(FocusCurves, InFocus, InZoom);
 }
 
 bool FFocalLengthTable::HasZoomPoint(float InFocus, float InZoom,  float InputTolerance)
@@ -465,6 +632,14 @@ bool FFocalLengthTable::HasZoomPoint(float InFocus, float InZoom,  float InputTo
 void FFocalLengthTable::ChangeZoomPoint(float InFocus, float InExistingZoom, float InNewZoom,  float InputTolerance)
 {
 	LensDataTableUtils::ChangeZoomPoint(FocusPoints, InFocus, InExistingZoom, InNewZoom, InputTolerance);
+
+	FFocalLengthInfo Data;
+	if (!GetPoint(InFocus, InNewZoom, Data, InputTolerance))
+	{
+		return;
+	}
+	
+	LensDataTableUtils::ChangeZoomInFocusCurves(FocusCurves, InFocus, InExistingZoom, InNewZoom, Data, InputTolerance);
 }
 
 bool FFocalLengthTable::DoesFocusPointExists(float InFocus, float InputTolerance) const
@@ -479,7 +654,13 @@ bool FFocalLengthTable::DoesFocusPointExists(float InFocus, float InputTolerance
 
 bool FFocalLengthTable::AddPoint(float InFocus, float InZoom, const FFocalLengthInfo& InData, float InputTolerance, bool bIsCalibrationPoint)
 {
-	return LensDataTableUtils::AddPoint(FocusPoints, InFocus, InZoom, InData, InputTolerance, bIsCalibrationPoint);
+	if (!LensDataTableUtils::AddPoint(FocusPoints, InFocus, InZoom, InData, InputTolerance, bIsCalibrationPoint))
+	{
+		return false;
+	}
+
+	LensDataTableUtils::AddPointToFocusCurve(FocusCurves, InFocus, InZoom, InData, InputTolerance);
+	return true;
 }
 
 bool FFocalLengthTable::GetPoint(const float InFocus, const float InZoom, FFocalLengthInfo& OutData, float InputTolerance) const
@@ -497,5 +678,18 @@ bool FFocalLengthTable::GetPoint(const float InFocus, const float InZoom, FFocal
 
 bool FFocalLengthTable::SetPoint(float InFocus, float InZoom, const FFocalLengthInfo& InData, float InputTolerance)
 {
-	return LensDataTableUtils::SetPoint(*this, InFocus, InZoom, InData, InputTolerance);
+	if (!LensDataTableUtils::SetPoint(*this, InFocus, InZoom, InData, InputTolerance))
+	{
+		return false;
+	}
+
+	LensDataTableUtils::SetPointInFocusCurve(FocusCurves, InFocus, InZoom, InData, InputTolerance);
+	return true;
+}
+
+void FFocalLengthTable::BuildFocusCurves()
+{
+	// Ensure that the focus curves are empty before building them from the table data
+	FocusCurves.Empty();
+	LensDataTableUtils::BuildFocusCurves(FocusPoints, FocusCurves);
 }
diff --git a/Engine/Plugins/VirtualProduction/CameraCalibrationCore/Source/CameraCalibrationCore/Private/Tables/ImageCenterTable.cpp b/Engine/Plugins/VirtualProduction/CameraCalibrationCore/Source/CameraCalibrationCore/Private/Tables/ImageCenterTable.cpp
--- a/Engine/Plugins/VirtualProduction/CameraCalibrationCore/Source/CameraCalibrationCore/Private/Tables/ImageCenterTable.cpp
+++ b/Engine/Plugins/VirtualProduction/CameraCalibrationCore/Source/CameraCalibrationCore/Private/Tables/ImageCenterTable.cpp
@@ -90,6 +90,87 @@ bool FImageCenterFocusPoint::IsEmpty() const
 	return Cx.IsEmpty();
 }
 
+const FRichCurve* FImageCenterFocusPoint::GetCurveForParameter(int32 InParameterIndex) const
+{
+	if (FImageCenterTable::FParameters::IsValid(InParameterIndex))
+	{
+		if (InParameterIndex == FImageCenterTable::FParameters::Cx)
+		{
+			return &Cx;
+		}
+
+		if (InParameterIndex == FImageCenterTable::FParameters::Cy)
+		{
+			return &Cy;
+		}
+	}
+
+	return nullptr;
+}
+
+FRichCurve* FImageCenterFocusPoint::GetCurveForParameter(int32 InParameterIndex)
+{
+	return const_cast<FRichCurve*>(const_cast<const FImageCenterFocusPoint*>(this)->GetCurveForParameter(InParameterIndex));
+}
+
+void FImageCenterFocusCurve::AddPoint(float InFocus, const FImageCenterInfo& InData, float InputTolerance)
+{
+	const FKeyHandle KeyHandle = AddPointToCurve(Cx, InFocus, InData.PrincipalPoint.X, InputTolerance);
+	AddPointToCurve(Cy, InFocus, InData.PrincipalPoint.Y, InputTolerance, KeyHandle);
+}
+
+void FImageCenterFocusCurve::SetPoint(float InFocus, const FImageCenterInfo& InData, float InputTolerance)
+{
+	SetPointInCurve(Cx, InFocus, InData.PrincipalPoint.X, InputTolerance);
+	SetPointInCurve(Cy, InFocus, InData.PrincipalPoint.Y, InputTolerance);
+}
+
+void FImageCenterFocusCurve::RemovePoint(float InFocus, float InputTolerance)
+{
+	DeletePointFromCurve(Cx, InFocus, InputTolerance);
+	DeletePointFromCurve(Cy, InFocus, InputTolerance);
+}
+
+void FImageCenterFocusCurve::ChangeFocus(float InExistingFocus, float InNewFocus, float InputTolerance)
+{
+	ChangeFocusInCurve(Cx, InExistingFocus, InNewFocus, InputTolerance);
+	ChangeFocusInCurve(Cy, InExistingFocus, InNewFocus, InputTolerance);
+}
+
+void FImageCenterFocusCurve::MergeFocus(float InExistingFocus, float InNewFocus, bool bReplaceExisting, float InputTolerance)
+{
+	MergeFocusInCurve(Cx, InExistingFocus, InNewFocus, bReplaceExisting, InputTolerance);
+	MergeFocusInCurve(Cy, InExistingFocus, InNewFocus, bReplaceExisting, InputTolerance);
+}
+
+bool FImageCenterFocusCurve::IsEmpty() const
+{
+	return !Cx.GetNumKeys() && !Cy.GetNumKeys();
+}
+
+const FRichCurve* FImageCenterFocusCurve::GetCurveForParameter(int32 InParameterIndex) const
+{
+	if (FImageCenterTable::FParameters::IsValid(InParameterIndex))
+	{
+		if (InParameterIndex == FImageCenterTable::FParameters::Cx)
+		{
+			return &Cx;
+		}
+
+		if (InParameterIndex == FImageCenterTable::FParameters::Cy)
+		{
+			return &Cy;
+		}
+	}
+
+	return nullptr;
+}
+
+FRichCurve* FImageCenterFocusCurve::GetCurveForParameter(int32 InParameterIndex)
+{
+	return const_cast<FRichCurve*>(const_cast<const FImageCenterFocusCurve*>(this)->GetCurveForParameter(InParameterIndex));
+}
+
 void FImageCenterTable::ForEachPoint(FFocusPointCallback InCallback) const
 {
 	for (const FImageCenterFocusPoint& Point : FocusPoints)
@@ -148,21 +229,22 @@ bool FImageCenterTable::BuildParameterCurveAtZoom(float InZoom, int32 InParamete
 	{
 		return false;
 	}
-	
-	for (const FImageCenterFocusPoint& FocusPoint : FocusPoints)
+
+	if (const FImageCenterFocusCurve* FocusCurve = GetFocusCurve(InZoom))
 	{
-		FImageCenterInfo ZoomPoint;
-		if (FocusPoint.GetPoint(InZoom, ZoomPoint))
+		if (InParameterIndex == FParameters::Cx)
+		{
+			OutCurve = FocusCurve->Cx;
+		}
+		else
 		{
-			const float Value = ZoomPoint.PrincipalPoint[InParameterIndex];
-			const FKeyHandle NewKeyHandle = OutCurve.AddKey(FocusPoint.Focus, Value);
-			FRichCurveKey& NewKey = OutCurve.GetKey(NewKeyHandle);
-			NewKey.TangentMode = ERichCurveTangentMode::RCTM_None;
-			NewKey.InterpMode = ERichCurveInterpMode::RCIM_Linear;
+			OutCurve = FocusCurve->Cy;
 		}
+
+		return true;
 	}
 
-	return true;
+	return false;
 }
 
 void FImageCenterTable::SetParameterCurveKeysAtFocus(float InFocus, int32 InParameterIndex, const FRichCurve& InSourceCurve, TArrayView<const FKeyHandle> InKeys)
@@ -174,27 +256,16 @@ void FImageCenterTable::SetParameterCurveKeysAtFocus(float InFocus, int32 InPara
 	
 	if (FImageCenterFocusPoint* FocusPoint = GetFocusPoint(InFocus))
 	{
-		FRichCurve* ActiveCurve = nullptr;
-		if (InParameterIndex == FParameters::Cx)
-		{
-			ActiveCurve = &FocusPoint->Cx;
-		}
-		else
+		CopyCurveKeys(InSourceCurve, *FocusPoint->GetCurveForParameter(InParameterIndex), InKeys);
+		PropagateCurveValuesToCrossCurves(*FocusPoint->GetCurveForParameter(InParameterIndex), InFocus, [this, InParameterIndex](float InZoom)->FRichCurve*
 		{
-			ActiveCurve = &FocusPoint->Cy;
-		}
-		
-		for (int32 Index = 0; Index < InKeys.Num(); ++Index)
-		{
-			const FKeyHandle Handle = InKeys[Index];
-			const int32 KeyIndex = InSourceCurve.GetIndexSafe(Handle);
-			if (KeyIndex != INDEX_NONE)
+			if (FImageCenterFocusCurve* Curve = GetFocusCurve(InZoom))
 			{
-				ActiveCurve->Keys[KeyIndex] = InSourceCurve.GetKey(Handle);
+				return Curve->GetCurveForParameter(InParameterIndex);
 			}
-		}
 
-		ActiveCurve->AutoSetTangents();
+			return nullptr;
+		});
 	}
 }
 
@@ -205,34 +276,18 @@ void FImageCenterTable::SetParameterCurveKeysAtZoom(float InZoom, int32 InParame
 		return;
 	}
 	
-	for (const FKeyHandle& KeyHandle : InKeys)
+	if (FImageCenterFocusCurve* FocusCurve = GetFocusCurve(InZoom))
 	{
-		// Assume the focus keys are put into the source curve in the same order as they are stored internally
-		const int32 KeyIndex = InSourceCurve.GetIndexSafe(KeyHandle);
-		if (KeyIndex != INDEX_NONE)
+		CopyCurveKeys(InSourceCurve, *FocusCurve->GetCurveForParameter(InParameterIndex), InKeys);
+		PropagateCurveValuesToCrossCurves(*FocusCurve->GetCurveForParameter(InParameterIndex), InZoom, [this, InParameterIndex](float InFocus)->FRichCurve*
 		{
-			if (ensure(FocusPoints.IsValidIndex(KeyIndex)))
+			if (FImageCenterFocusPoint* FocusPoint = GetFocusPoint(InFocus))
 			{
-				FImageCenterFocusPoint& FocusPoint = FocusPoints[KeyIndex];
-				FImageCenterInfo ZoomPoint;
-				if (!FocusPoint.GetPoint(InZoom, ZoomPoint))
-				{
-					continue;
-				}
-
-				ZoomPoint.PrincipalPoint[InParameterIndex] = InSourceCurve.GetKeyValue(KeyHandle);
-				FocusPoint.SetPoint(InZoom, ZoomPoint);
-
-				if (InParameterIndex == FParameters::Cx)
-				{
-					FocusPoint.Cx.AutoSetTangents();
-				}
-				else
-				{
-					FocusPoint.Cy.AutoSetTangents();
-				}
+				return FocusPoint->GetCurveForParameter(InParameterIndex);
 			}
-		}
+
+			return nullptr;
+		});
 	}
 }
 
@@ -251,6 +306,16 @@ FImageCenterFocusPoint* FImageCenterTable::GetFocusPoint(float InFocus, float In
 	return FocusPoints.FindByPredicate([InFocus, InputTolerance](const FImageCenterFocusPoint& Points) { return FMath::IsNearlyEqual(Points.Focus, InFocus, InputTolerance); });
 }
 
+const FImageCenterFocusCurve* FImageCenterTable::GetFocusCurve(float InZoom, float InputTolerance) const
+{
+	return FocusCurves.FindByPredicate([InZoom, InputTolerance](const FImageCenterFocusCurve& Curve) { return FMath::IsNearlyEqual(Curve.Zoom, InZoom, InputTolerance); });
+}
+
+FImageCenterFocusCurve* FImageCenterTable::GetFocusCurve(float InZoom, float InputTolerance)
+{
+	return FocusCurves.FindByPredicate([InZoom, InputTolerance](const FImageCenterFocusCurve& Curve) { return FMath::IsNearlyEqual(Curve.Zoom, InZoom, InputTolerance); });
+}
+
 TConstArrayView<FImageCenterFocusPoint> FImageCenterTable::GetFocusPoints() const
 {
 	return FocusPoints;
@@ -261,6 +326,16 @@ TArray<FImageCenterFocusPoint>& FImageCenterTable::GetFocusPoints()
 	return FocusPoints;
 }
 
+TConstArrayView<FImageCenterFocusCurve> FImageCenterTable::GetFocusCurves() const
+{
+	return FocusCurves;
+}
+
+TArray<FImageCenterFocusCurve>& FImageCenterTable::GetFocusCurves()
+{
+	return FocusCurves;
+}
+
 bool FImageCenterTable::DoesZoomPointExists(float InFocus, float InZoom, float InputTolerance) const
 {
 	FImageCenterInfo ImageCenterInfo;
@@ -297,6 +372,7 @@ TMap<ELensDataCategory, FLinkPointMetadata> FImageCenterTable::GetLinkedCategori
 void FImageCenterTable::RemoveFocusPoint(float InFocus)
 {
 	LensDataTableUtils::RemoveFocusPoint(FocusPoints, InFocus);
+	LensDataTableUtils::RemoveFocusFromFocusCurves(FocusCurves, InFocus);
 }
 
 bool FImageCenterTable::HasFocusPoint(float InFocus, float InputTolerance) const
@@ -307,16 +383,19 @@ bool FImageCenterTable::HasFocusPoint(float InFocus, float InputTolerance) const
 void FImageCenterTable::ChangeFocusPoint(float InExistingFocus, float InNewFocus, float InputTolerance)
 {
 	LensDataTableUtils::ChangeFocusPoint(FocusPoints, InExistingFocus, InNewFocus, InputTolerance);
+	LensDataTableUtils::ChangeFocusInFocusCurves(FocusCurves, InExistingFocus, InNewFocus, InputTolerance);
 }
 
 void FImageCenterTable::MergeFocusPoint(float InSrcFocus, float InDestFocus, bool bReplaceExistingZoomPoints, float InputTolerance)
 {
 	LensDataTableUtils::MergeFocusPoint(FocusPoints, InSrcFocus, InDestFocus, bReplaceExistingZoomPoints, InputTolerance);
+	LensDataTableUtils::MergeFocusInFocusCurves(FocusCurves, InSrcFocus, InDestFocus, bReplaceExistingZoomPoints, InputTolerance);
 }
 
 void FImageCenterTable::RemoveZoomPoint(float InFocus, float InZoom)
 {
 	LensDataTableUtils::RemoveZoomPoint(FocusPoints, InFocus, InZoom);
+	LensDataTableUtils::RemoveZoomFromFocusCurves(FocusCurves, InFocus, InZoom);
 }
 
 bool FImageCenterTable::HasZoomPoint(float InFocus, float InZoom, float InputTolerance)
@@ -327,12 +406,26 @@ bool FImageCenterTable::HasZoomPoint(float InFocus, float InZoom, float InputTol
 void FImageCenterTable::ChangeZoomPoint(float InFocus, float InExistingZoom, float InNewZoom, float InputTolerance)
 {
 	LensDataTableUtils::ChangeZoomPoint(FocusPoints, InFocus, InExistingZoom, InNewZoom, InputTolerance);
+	
+	FImageCenterInfo Data;
+	if (!GetPoint(InFocus, InNewZoom, Data, InputTolerance))
+	{
+		return;
+	}
+
+	LensDataTableUtils::ChangeZoomInFocusCurves(FocusCurves, InFocus, InExistingZoom, InNewZoom, Data, InputTolerance);
 }
 
 bool FImageCenterTable::AddPoint(float InFocus, float InZoom, const FImageCenterInfo& InData, float InputTolerance,
 	bool bIsCalibrationPoint)
 {
-	return LensDataTableUtils::AddPoint(FocusPoints, InFocus, InZoom, InData, InputTolerance, bIsCalibrationPoint);
+	if (!LensDataTableUtils::AddPoint(FocusPoints, InFocus, InZoom, InData, InputTolerance, bIsCalibrationPoint))
+	{
+		return false;
+	}
+
+	LensDataTableUtils::AddPointToFocusCurve(FocusCurves, InFocus, InZoom, InData, InputTolerance);
+	return true;
 }
 
 bool FImageCenterTable::GetPoint(const float InFocus, const float InZoom, FImageCenterInfo& OutData, float InputTolerance) const
@@ -353,5 +446,19 @@ bool FImageCenterTable::GetPoint(const float InFocus, const float InZoom, FImage
 
 bool FImageCenterTable::SetPoint(float InFocus, float InZoom, const FImageCenterInfo& InData, float InputTolerance)
 {
-	return LensDataTableUtils::SetPoint(*this, InFocus, InZoom, InData, InputTolerance);
+	if (!LensDataTableUtils::SetPoint(*this, InFocus, InZoom, InData, InputTolerance))
+	{
+		return false;
+	}
+
+	LensDataTableUtils::SetPointInFocusCurve(FocusCurves, InFocus, InZoom, InData, InputTolerance);
+	
+	return true;
+}
+
+void FImageCenterTable::BuildFocusCurves()
+{
+	// Ensure that the focus curves are empty before building them from the table data
+	FocusCurves.Empty();
+	LensDataTableUtils::BuildFocusCurves(FocusPoints, FocusCurves);
 }
diff --git a/Engine/Plugins/VirtualProduction/CameraCalibrationCore/Source/CameraCalibrationCore/Private/Tables/LensTableUtils.h b/Engine/Plugins/VirtualProduction/CameraCalibrationCore/Source/CameraCalibrationCore/Private/Tables/LensTableUtils.h
--- a/Engine/Plugins/VirtualProduction/CameraCalibrationCore/Source/CameraCalibrationCore/Private/Tables/LensTableUtils.h
+++ b/Engine/Plugins/VirtualProduction/CameraCalibrationCore/Source/CameraCalibrationCore/Private/Tables/LensTableUtils.h
@@ -21,6 +21,15 @@ namespace LensDataTableUtils
     	}
 	}
 
+	template<typename FocusCurveType>
+	void RemoveFocusFromFocusCurves(TArray<FocusCurveType>& InFocusCurves, float InFocus)
+	{
+		for (FocusCurveType& Curve : InFocusCurves)
+		{
+			Curve.RemovePoint(InFocus, UE_SMALL_NUMBER);
+		}
+	}
+	
 	/** Changes the value of a focus point in the container */
 	template<typename FocusPointType>
 	void ChangeFocusPoint(TArray<FocusPointType>& Container, float InExistingFocus, float InNewFocus, float InputTolerance = KINDA_SMALL_NUMBER)
@@ -36,6 +45,15 @@ namespace LensDataTableUtils
 		}
 	}
 
+	template<typename FocusCurveType>
+	void ChangeFocusInFocusCurves(TArray<FocusCurveType>& InFocusCurves, float InExistingFocus, float InNewFocus, float InputTolerance = KINDA_SMALL_NUMBER)
+	{
+		for (FocusCurveType& Curve : InFocusCurves)
+		{
+			Curve.ChangeFocus(InExistingFocus, InNewFocus, InputTolerance);
+		}
+	}
+	
 	/** Merges the points in the specified source focus into the specified destination focus */
 	template<typename FocusPointType>
 	void MergeFocusPoint(TArray<FocusPointType>& Container, float InSrcFocus, float InDestFocus, bool bReplaceExistingZoomPoints, float InputTolerance = KINDA_SMALL_NUMBER)
@@ -88,7 +106,16 @@ namespace LensDataTableUtils
 			}
 		}
 	}
-
+	
+	template<typename FocusCurveType>
+	void MergeFocusInFocusCurves(TArray<FocusCurveType>& InFocusCurves, float InSrcFocus, float InDestFocus, bool bReplaceExistingZoomPoints, float InputTolerance = KINDA_SMALL_NUMBER)
+	{
+		for (FocusCurveType& Curve : InFocusCurves)
+		{
+			Curve.MergeFocus(InSrcFocus, InDestFocus, bReplaceExistingZoomPoints, InputTolerance);
+		}
+	}
+	
 	/** Gets all point info for specific data table */
 	template<typename FPointInfoType, typename FDataTableType>
 	TArray<FPointInfoType> GetAllPointsInfo(const FDataTableType& InTable)
@@ -137,6 +164,25 @@ namespace LensDataTableUtils
 		}
 	}
 
+	template<typename FocusCurveType>
+	void RemoveZoomFromFocusCurves(TArray<FocusCurveType>& FocusCurves, float InFocus, float InZoom, float InputTolerance = KINDA_SMALL_NUMBER)
+	{
+		const int32 CurveIndex = FocusCurves.IndexOfByPredicate([InZoom, InputTolerance](const FocusCurveType& Curve)
+		{
+			return FMath::IsNearlyEqual(Curve.Zoom, InZoom, InputTolerance);
+		});
+		
+		if (CurveIndex != INDEX_NONE)
+		{
+			FocusCurves[CurveIndex].RemovePoint(InFocus, InputTolerance);
+			if (FocusCurves[CurveIndex].IsEmpty())
+			{
+				// If the curve has no points, that indicates there are no focuses that contain the zoom value, so we can delete it entirely
+				FocusCurves.RemoveAt(CurveIndex);
+			}
+		}
+	}
+	
 	/** Changes the value of a zoom point for a given focus value in a container */
 	template<typename FocusPointType>
 	void ChangeZoomPoint(TArray<FocusPointType>& Container, float InFocus, float InExistingZoom, float InNewZoom, float InputTolerance = KINDA_SMALL_NUMBER)
@@ -158,6 +204,29 @@ namespace LensDataTableUtils
 		}
 	}
 	
+	template<typename FocusCurveType, typename DataType>
+	void ChangeZoomInFocusCurves(TArray<FocusCurveType>& FocusCurves, float InFocus, float InExistingZoom, float InNewZoom, const DataType& InData, float InputTolerance = KINDA_SMALL_NUMBER)
+	{
+		RemoveZoomFromFocusCurves(FocusCurves, InFocus, InExistingZoom, InputTolerance);
+		
+		const int32 NewCurveIndex = FocusCurves.IndexOfByPredicate([InNewZoom, InputTolerance](const FocusCurveType& Curve)
+		{
+			return FMath::IsNearlyEqual(Curve.Zoom, InNewZoom, InputTolerance);
+		});
+		
+		if (NewCurveIndex != INDEX_NONE)
+		{
+			FocusCurves[NewCurveIndex].AddPoint(InFocus, InData, InputTolerance);
+			return;
+		}
+
+		// At this point, no focus curve has been found that matches the input zoom to the specified tolerance, so create one
+		FocusCurveType NewCurve;
+		NewCurve.Zoom = InNewZoom;
+		NewCurve.AddPoint(InFocus, InData, InputTolerance);
+		FocusCurves.Add(NewCurve);
+	}
+	
 	/** Adds a point at a specified focus and zoom input values */
 	template<typename FocusPointType, typename DataType>
 	bool AddPoint(TArray<FocusPointType>& InContainer, float InFocus, float InZoom, const DataType& InData, float InputTolerance, bool bIsCalibrationPoint)
@@ -186,6 +255,33 @@ namespace LensDataTableUtils
 
 		return bSuccess;
 	}
+
+	/** Adds a point at the specified focus and zoom to a matching focus curve */
+	template<typename FocusCurveType, typename DataType>
+	void AddPointToFocusCurve(TArray<FocusCurveType>& InContainer, float InFocus, float InZoom, const DataType& InData, float InputTolerance)
+	{
+		int32 CurveIndex = 0;
+		for (; CurveIndex < InContainer.Num(); ++CurveIndex)
+		{
+			FocusCurveType& FocusCurve = InContainer[CurveIndex];
+			if (FMath::IsNearlyEqual(FocusCurve.Zoom, InZoom, InputTolerance))
+			{
+				FocusCurve.AddPoint(InFocus, InData, InputTolerance);
+				return;
+			}
+
+			if (InZoom < FocusCurve.Zoom)
+			{
+				break;
+			}
+		}
+
+		// At this point, no focus curve has been found that matches the input zoom to the specified tolerance, so create one
+		FocusCurveType NewCurve;
+		NewCurve.Zoom = InZoom;
+		NewCurve.AddPoint(InFocus, InData, InputTolerance);
+		InContainer.Insert(MoveTemp(NewCurve), CurveIndex);
+	}
 	
 	template<typename TableType, typename DataType>
 	bool SetPoint(TableType& InTable, float InFocus, float InZoom, const DataType& InData, float InputTolerance = KINDA_SMALL_NUMBER)
@@ -202,6 +298,26 @@ namespace LensDataTableUtils
 		return false;
 	}
 
+	/** Sets a point at the specified focus and zoom to a matching focus curve */
+	template<typename FocusCurveType, typename DataType>
+	void SetPointInFocusCurve(TArray<FocusCurveType>& InContainer, float InFocus, float InZoom, const DataType& InData, float InputTolerance)
+	{
+		for (int32 CurveIndex = 0; CurveIndex < InContainer.Num(); ++CurveIndex)
+		{
+			FocusCurveType& FocusCurve = InContainer[CurveIndex];
+			if (FMath::IsNearlyEqual(FocusCurve.Zoom, InZoom, InputTolerance))
+			{
+				FocusCurve.SetPoint(InFocus, InData, InputTolerance);
+				return;
+			}
+
+			if (InZoom < FocusCurve.Zoom)
+			{
+				break;
+			}
+		}
+	}
+	
 	/** Clears content of a table */
 	template<typename Type>
 	void EmptyTable(Type& InTable)
@@ -209,8 +325,11 @@ namespace LensDataTableUtils
 		InTable.FocusPoints.Empty(0);
 	}
 
-		struct FPointNeighbors
+	struct FPointNeighbors
 	{
+		/** Returns true if the previous index is identical to the next index */
+		bool IsSinglePoint() const { return PreviousIndex == NextIndex; }
+		
 		int32 PreviousIndex = INDEX_NONE;
 		int32 NextIndex = INDEX_NONE;
 	};
@@ -253,6 +372,44 @@ namespace LensDataTableUtils
 		return Neighbors;
 	}
 
+	template<typename Type>
+	FPointNeighbors FindFocusCurves(float InZoom, TConstArrayView<Type> Container)
+	{
+		FPointNeighbors Neighbors;
+		if (Container.Num() <= 0)
+		{
+			return Neighbors;
+		}
+
+		for (int32 Index = 0; Index < Container.Num(); ++Index)
+		{
+			const Type& Curve = Container[Index];
+			if (Curve.Zoom > InZoom)
+			{
+				Neighbors.NextIndex = Index;
+				Neighbors.PreviousIndex = FMath::Max(Index - 1, 0);
+				break;
+			}
+			
+			if (FMath::IsNearlyEqual(Curve.Zoom, InZoom))
+			{
+				//We found a point exactly matching the desired one
+				Neighbors.NextIndex = Index;
+				Neighbors.PreviousIndex = Index;
+				break;
+			}
+		}
+
+		//We haven't found a curve, default to last one
+		if (Neighbors.PreviousIndex == INDEX_NONE && Neighbors.NextIndex == INDEX_NONE)
+		{
+			Neighbors.NextIndex = Container.Num() - 1;
+			Neighbors.PreviousIndex = Container.Num() - 1;
+		}
+
+		return Neighbors;
+	}
+	
 	/** Finds indices of neighbor zoom points for a given zoom value */
 	template<typename Type>
 	FPointNeighbors FindZoomPoints(float InZoom, const TArray<Type>& Container)
@@ -337,4 +494,48 @@ namespace LensDataTableUtils
 
 		return PointNum;
 	}
+
+	template<typename FocusPointType, typename FocusCurveType>
+	void BuildFocusCurves(const TArray<FocusPointType>& InPoints, TArray<FocusCurveType>& OutCurves)
+	{
+		for (const FocusPointType& FocusPoint : InPoints)
+		{
+			for (int32 Index = 0; Index < FocusPoint.GetNumPoints(); ++Index)
+			{
+				const float Zoom = FocusPoint.GetZoom(Index);
+
+				typename FocusPointType::PointType PointData;
+				if (FocusPoint.GetPoint(Zoom, PointData))
+				{
+					int32 ExistingCurveIndex = 0;
+					for (; ExistingCurveIndex < OutCurves.Num(); ++ExistingCurveIndex)
+					{
+						FocusCurveType& Curve = OutCurves[ExistingCurveIndex];
+						if (FMath::IsNearlyEqual(Curve.Zoom, Zoom))
+						{
+							Curve.AddPoint(FocusPoint.Focus, PointData);
+							break;
+						}
+
+						if (Zoom < Curve.Zoom)
+						{
+							FocusCurveType NewCurve;
+							NewCurve.Zoom = Zoom;
+							NewCurve.AddPoint(FocusPoint.Focus, PointData);
+							OutCurves.Insert(NewCurve, ExistingCurveIndex);
+							break;
+						}
+					}
+
+					if (ExistingCurveIndex == OutCurves.Num())
+					{
+						FocusCurveType NewCurve;
+						NewCurve.Zoom = Zoom;
+						NewCurve.AddPoint(FocusPoint.Focus, PointData);
+						OutCurves.Add(NewCurve);
+					}
+				}
+			}
+		}
+	}
 }
diff --git a/Engine/Plugins/VirtualProduction/CameraCalibrationCore/Source/CameraCalibrationCore/Private/Tables/NodalOffsetTable.cpp b/Engine/Plugins/VirtualProduction/CameraCalibrationCore/Source/CameraCalibrationCore/Private/Tables/NodalOffsetTable.cpp
--- a/Engine/Plugins/VirtualProduction/CameraCalibrationCore/Source/CameraCalibrationCore/Private/Tables/NodalOffsetTable.cpp
+++ b/Engine/Plugins/VirtualProduction/CameraCalibrationCore/Source/CameraCalibrationCore/Private/Tables/NodalOffsetTable.cpp
@@ -142,6 +142,132 @@ bool FNodalOffsetFocusPoint::IsEmpty() const
 	return LocationOffset[0].IsEmpty();
 }
 
+const FRichCurve* FNodalOffsetFocusPoint::GetCurveForParameter(int32 InParameterIndex) const
+{
+	if (FNodalOffsetTable::FParameters::IsValidComposed(InParameterIndex))
+	{
+		int32 Parameter;
+		EAxis::Type Axis;
+		FNodalOffsetTable::FParameters::Decompose(InParameterIndex, Parameter, Axis);
+
+		if (Parameter == FNodalOffsetTable::FParameters::Location)
+		{
+			return &LocationOffset[Axis - 1];
+		}
+
+		if (Parameter == FNodalOffsetTable::FParameters::Rotation)
+		{
+			return &RotationOffset[Axis - 1];
+		}
+	}
+
+	return nullptr;
+}
+
+FRichCurve* FNodalOffsetFocusPoint::GetCurveForParameter(int32 InParameterIndex)
+{
+	return const_cast<FRichCurve*>(const_cast<const FNodalOffsetFocusPoint*>(this)->GetCurveForParameter(InParameterIndex));
+}
+
+void FNodalOffsetFocusCurve::AddPoint(float InFocus, const FNodalPointOffset& InData, float InputTolerance)
+{
+	for (int32 Index = 0; Index < FNodalOffsetFocusPoint::LocationDimension; ++Index)
+	{
+		AddPointToCurve(LocationOffset[Index], InFocus, InData.LocationOffset[Index], InputTolerance);
+	}
+
+	const FRotator NewRotator = InData.RotationOffset.Rotator();
+	for (int32 Index = 0; Index < FNodalOffsetFocusPoint::RotationDimension; ++Index)
+	{
+		AddPointToCurve(RotationOffset[Index], InFocus, NewRotator.GetComponentForAxis(static_cast<EAxis::Type>(Index+1)), InputTolerance);
+	}
+}
+
+void FNodalOffsetFocusCurve::SetPoint(float InFocus, const FNodalPointOffset& InData, float InputTolerance)
+{
+	for (int32 Index = 0; Index < FNodalOffsetFocusPoint::LocationDimension; ++Index)
+	{
+		SetPointInCurve(LocationOffset[Index], InFocus, InData.LocationOffset[Index], InputTolerance);
+	}
+
+	const FRotator NewRotator = InData.RotationOffset.Rotator();
+	for (int32 Index = 0; Index < FNodalOffsetFocusPoint::RotationDimension; ++Index)
+	{
+		SetPointInCurve(RotationOffset[Index], InFocus, NewRotator.GetComponentForAxis(static_cast<EAxis::Type>(Index+1)), InputTolerance);
+	}
+}
+
+void FNodalOffsetFocusCurve::RemovePoint(float InFocus, float InputTolerance)
+{
+	for (int32 Index = 0; Index < FNodalOffsetFocusPoint::LocationDimension; ++Index)
+	{
+		DeletePointFromCurve(LocationOffset[Index], InFocus, InputTolerance);
+	}
+	
+	for (int32 Index = 0; Index < FNodalOffsetFocusPoint::RotationDimension; ++Index)
+	{
+		DeletePointFromCurve(RotationOffset[Index], InFocus, InputTolerance);
+	}
+}
+
+void FNodalOffsetFocusCurve::ChangeFocus(float InExistingFocus, float InNewFocus, float InputTolerance)
+{
+	for (int32 Index = 0; Index < FNodalOffsetFocusPoint::LocationDimension; ++Index)
+	{
+		ChangeFocusInCurve(LocationOffset[Index], InExistingFocus, InNewFocus, InputTolerance);
+	}
+	
+	for (int32 Index = 0; Index < FNodalOffsetFocusPoint::RotationDimension; ++Index)
+	{
+		ChangeFocusInCurve(RotationOffset[Index], InExistingFocus, InNewFocus, InputTolerance);
+	}
+}
+
+void FNodalOffsetFocusCurve::MergeFocus(float InExistingFocus, float InNewFocus, bool bReplaceExisting, float InputTolerance)
+{
+	for (int32 Index = 0; Index < FNodalOffsetFocusPoint::LocationDimension; ++Index)
+	{
+		MergeFocusInCurve(LocationOffset[Index], InExistingFocus, InNewFocus, bReplaceExisting, InputTolerance);
+	}
+	
+	for (int32 Index = 0; Index < FNodalOffsetFocusPoint::RotationDimension; ++Index)
+	{
+		MergeFocusInCurve(RotationOffset[Index], InExistingFocus, InNewFocus, bReplaceExisting, InputTolerance);
+	}
+}
+
+bool FNodalOffsetFocusCurve::IsEmpty() const
+{
+	return !LocationOffset[0].GetNumKeys();
+}
+
+const FRichCurve* FNodalOffsetFocusCurve::GetCurveForParameter(int32 InParameterIndex) const
+{
+	if (FNodalOffsetTable::FParameters::IsValidComposed(InParameterIndex))
+	{
+		int32 Parameter;
+		EAxis::Type Axis;
+		FNodalOffsetTable::FParameters::Decompose(InParameterIndex, Parameter, Axis);
+
+		if (Parameter == FNodalOffsetTable::FParameters::Location)
+		{
+			return &LocationOffset[Axis - 1];
+		}
+
+		if (Parameter == FNodalOffsetTable::FParameters::Rotation)
+		{
+			return &RotationOffset[Axis - 1];
+		}
+	}
+
+	return nullptr;
+}
+
+FRichCurve* FNodalOffsetFocusCurve::GetCurveForParameter(int32 InParameterIndex)
+{
+	return const_cast<FRichCurve*>(const_cast<const FNodalOffsetFocusCurve*>(this)->GetCurveForParameter(InParameterIndex));
+}
+
 bool FNodalOffsetTable::DoesZoomPointExists(float InFocus, float InZoom, float InputTolerance) const
 {
 	FNodalPointOffset NodalPointOffset;
@@ -223,21 +349,22 @@ bool FNodalOffsetTable::BuildParameterCurveAtZoom(float InZoom, int32 InParamete
 	int32 Parameter;
 	EAxis::Type Axis;
 	FParameters::Decompose(InParameterIndex, Parameter, Axis);
-	
-	for (const FNodalOffsetFocusPoint& FocusPoint : FocusPoints)
+
+	if (const FNodalOffsetFocusCurve* FocusCurve = GetFocusCurve(InZoom))
 	{
-		FNodalPointOffset ZoomPoint;
-		if (FocusPoint.GetPoint(InZoom, ZoomPoint))
+		if (Parameter == FParameters::Location)
 		{
-			const float Value = Parameter == 0 ? ZoomPoint.LocationOffset[Axis - 1] : ZoomPoint.RotationOffset.Rotator().GetComponentForAxis(Axis);
-			const FKeyHandle NewKeyHandle = OutCurve.AddKey(FocusPoint.Focus, Value);
-			FRichCurveKey& NewKey = OutCurve.GetKey(NewKeyHandle);
-			NewKey.TangentMode = ERichCurveTangentMode::RCTM_None;
-			NewKey.InterpMode = ERichCurveInterpMode::RCIM_Linear;
+			OutCurve = FocusCurve->LocationOffset[Axis - 1];
 		}
+		else
+		{
+			OutCurve = FocusCurve->RotationOffset[Axis - 1];
+		}
+		
+		return true;
 	}
 
-	return true;
+	return false;
 }
 
 void FNodalOffsetTable::SetParameterCurveKeysAtFocus(float InFocus, int32 InParameterIndex, const FRichCurve& InSourceCurve, TArrayView<const FKeyHandle> InKeys)
@@ -249,31 +376,16 @@ void FNodalOffsetTable::SetParameterCurveKeysAtFocus(float InFocus, int32 InPara
 	
 	if (FNodalOffsetFocusPoint* FocusPoint = GetFocusPoint(InFocus))
 	{
-		int32 Parameter;
-		EAxis::Type Axis;
-		FParameters::Decompose(InParameterIndex, Parameter, Axis);
-		
-		FRichCurve* ActiveCurve = nullptr;
-		if (Parameter == FParameters::Location)
+		CopyCurveKeys(InSourceCurve, *FocusPoint->GetCurveForParameter(InParameterIndex), InKeys);
+		PropagateCurveValuesToCrossCurves(*FocusPoint->GetCurveForParameter(InParameterIndex), InFocus, [this, InParameterIndex](float InZoom)->FRichCurve*
 		{
-			ActiveCurve = &FocusPoint->LocationOffset[Axis - 1];
-		}
-		else
-		{
-			ActiveCurve = &FocusPoint->RotationOffset[Axis - 1];
-		}
-		
-		for (int32 Index = 0; Index < InKeys.Num(); ++Index)
-		{
-			const FKeyHandle Handle = InKeys[Index];
-			const int32 KeyIndex = InSourceCurve.GetIndexSafe(Handle);
-			if (KeyIndex != INDEX_NONE)
+			if (FNodalOffsetFocusCurve* Curve = GetFocusCurve(InZoom))
 			{
-				ActiveCurve->Keys[KeyIndex] = InSourceCurve.GetKey(Handle);
+				return Curve->GetCurveForParameter(InParameterIndex);
 			}
-		}
 
-		ActiveCurve->AutoSetTangents();
+			return nullptr;
+		});
 	}
 }
 
@@ -284,48 +396,18 @@ void FNodalOffsetTable::SetParameterCurveKeysAtZoom(float InZoom, int32 InParame
 		return;
 	}
 
-	int32 Parameter;
-	EAxis::Type Axis;
-	FParameters::Decompose(InParameterIndex, Parameter, Axis);
-	
-	for (const FKeyHandle& KeyHandle : InKeys)
+	if (FNodalOffsetFocusCurve* FocusCurve = GetFocusCurve(InZoom))
 	{
-		// Assume the focus keys are put into the source curve in the same order as they are stored internally
-		const int32 KeyIndex = InSourceCurve.GetIndexSafe(KeyHandle);
-		if (KeyIndex != INDEX_NONE)
+		CopyCurveKeys(InSourceCurve, *FocusCurve->GetCurveForParameter(InParameterIndex), InKeys);
+		PropagateCurveValuesToCrossCurves(*FocusCurve->GetCurveForParameter(InParameterIndex), InZoom, [this, InParameterIndex](float InFocus)->FRichCurve*
 		{
-			if (ensure(FocusPoints.IsValidIndex(KeyIndex)))
+			if (FNodalOffsetFocusPoint* Point = GetFocusPoint(InFocus))
 			{
-				FNodalOffsetFocusPoint& FocusPoint = FocusPoints[KeyIndex];
-				FNodalPointOffset ZoomPoint;
-				if (!FocusPoint.GetPoint(InZoom, ZoomPoint))
-				{
-					continue;
-				}
-
-				if (Parameter == FParameters::Location)
-				{
-					ZoomPoint.LocationOffset[Axis - 1] = InSourceCurve.GetKeyValue(KeyHandle);
-				}
-				else if (Parameter == FParameters::Rotation)
-				{
-					FRotator Rotator = ZoomPoint.RotationOffset.Rotator();
-					Rotator.SetComponentForAxis(Axis, InSourceCurve.GetKeyValue(KeyHandle));
-					ZoomPoint.RotationOffset = Rotator.Quaternion();
-				}
-
-				FocusPoint.SetPoint(InZoom, ZoomPoint);
-
-				if (Parameter == FParameters::Location)
-				{
-					FocusPoint.LocationOffset[Axis - 1].AutoSetTangents();
-				}
-				else if (Parameter == FParameters::Rotation)
-				{
-					FocusPoint.RotationOffset[Axis - 1].AutoSetTangents();
-				}
+				return Point->GetCurveForParameter(InParameterIndex);
 			}
-		}
+
+			return nullptr;
+		});
 	}
 }
 
@@ -381,6 +463,16 @@ FNodalOffsetFocusPoint* FNodalOffsetTable::GetFocusPoint(float InFocus, float In
 	return FocusPoints.FindByPredicate([InFocus, InputTolerance](const FNodalOffsetFocusPoint& Point) { return FMath::IsNearlyEqual(Point.Focus, InFocus, InputTolerance); });
 }
 
+const FNodalOffsetFocusCurve* FNodalOffsetTable::GetFocusCurve(float InZoom, float InputTolerance) const
+{
+	return FocusCurves.FindByPredicate([InZoom, InputTolerance](const FNodalOffsetFocusCurve& Curve) { return FMath::IsNearlyEqual(Curve.Zoom, InZoom, InputTolerance); });
+}
+
+FNodalOffsetFocusCurve* FNodalOffsetTable::GetFocusCurve(float InZoom, float InputTolerance)
+{
+	return FocusCurves.FindByPredicate([InZoom, InputTolerance](const FNodalOffsetFocusCurve& Curve) { return FMath::IsNearlyEqual(Curve.Zoom, InZoom, InputTolerance); });
+}
+
 TConstArrayView<FNodalOffsetFocusPoint> FNodalOffsetTable::GetFocusPoints() const
 {
 	return FocusPoints;
@@ -391,6 +483,16 @@ TArray<FNodalOffsetFocusPoint>& FNodalOffsetTable::GetFocusPoints()
 	return FocusPoints;
 }
 
+TConstArrayView<FNodalOffsetFocusCurve> FNodalOffsetTable::GetFocusCurves() const
+{
+	return FocusCurves;
+}
+
+TArray<FNodalOffsetFocusCurve>& FNodalOffsetTable::GetFocusCurves()
+{
+	return FocusCurves;
+}
+
 void FNodalOffsetTable::ForEachPoint(FFocusPointCallback InCallback) const
 {
 	for (const FNodalOffsetFocusPoint& Point : FocusPoints)
@@ -402,6 +504,7 @@ void FNodalOffsetTable::ForEachPoint(FFocusPointCallback InCallback) const
 void FNodalOffsetTable::RemoveFocusPoint(float InFocus)
 {
 	LensDataTableUtils::RemoveFocusPoint(FocusPoints, InFocus);
+	LensDataTableUtils::RemoveFocusFromFocusCurves(FocusCurves, InFocus);
 }
 
 bool FNodalOffsetTable::HasFocusPoint(float InFocus, float InputTolerance) const
@@ -412,16 +515,19 @@ bool FNodalOffsetTable::HasFocusPoint(float InFocus, float InputTolerance) const
 void FNodalOffsetTable::ChangeFocusPoint(float InExistingFocus, float InNewFocus, float InputTolerance)
 {
 	LensDataTableUtils::ChangeFocusPoint(FocusPoints, InExistingFocus, InNewFocus, InputTolerance);
+	LensDataTableUtils::ChangeFocusInFocusCurves(FocusCurves, InExistingFocus, InNewFocus, InputTolerance);
 }
 
 void FNodalOffsetTable::MergeFocusPoint(float InSrcFocus, float InDestFocus, bool bReplaceExistingZoomPoints, float InputTolerance)
 {
 	LensDataTableUtils::MergeFocusPoint(FocusPoints, InSrcFocus, InDestFocus, bReplaceExistingZoomPoints, InputTolerance);
+	LensDataTableUtils::MergeFocusInFocusCurves(FocusCurves, InSrcFocus, InDestFocus, bReplaceExistingZoomPoints, InputTolerance);
 }
 
 void FNodalOffsetTable::RemoveZoomPoint(float InFocus, float InZoom)
 {
 	LensDataTableUtils::RemoveZoomPoint(FocusPoints, InFocus, InZoom);
+	LensDataTableUtils::RemoveZoomFromFocusCurves(FocusCurves, InFocus, InZoom);
 }
 
 bool FNodalOffsetTable::HasZoomPoint(float InFocus, float InZoom, float InputTolerance)
@@ -432,6 +538,14 @@ bool FNodalOffsetTable::HasZoomPoint(float InFocus, float InZoom, float InputTol
 void FNodalOffsetTable::ChangeZoomPoint(float InFocus, float InExistingZoom, float InNewZoom, float InputTolerance)
 {
 	LensDataTableUtils::ChangeZoomPoint(FocusPoints, InFocus, InExistingZoom, InNewZoom, InputTolerance);
+	
+	FNodalPointOffset Data;
+	if (!GetPoint(InFocus, InNewZoom, Data, InputTolerance))
+	{
+		return;
+	}
+
+	LensDataTableUtils::ChangeZoomInFocusCurves(FocusCurves, InFocus, InExistingZoom, InNewZoom, Data, InputTolerance);
 }
 
 bool FNodalOffsetTable::DoesFocusPointExists(float InFocus, float InputTolerance) const
@@ -446,7 +560,13 @@ bool FNodalOffsetTable::DoesFocusPointExists(float InFocus, float InputTolerance
 
 bool FNodalOffsetTable::AddPoint(float InFocus, float InZoom, const FNodalPointOffset& InData, float InputTolerance, bool bIsCalibrationPoint)
 {
-	return LensDataTableUtils::AddPoint(FocusPoints, InFocus, InZoom, InData, InputTolerance, bIsCalibrationPoint);
+	if (!LensDataTableUtils::AddPoint(FocusPoints, InFocus, InZoom, InData, InputTolerance, bIsCalibrationPoint))
+	{
+		return false;
+	}
+
+	LensDataTableUtils::AddPointToFocusCurve(FocusCurves, InFocus, InZoom, InData, InputTolerance);
+	return true;
 }
 
 bool FNodalOffsetTable::GetPoint(const float InFocus, const float InZoom, FNodalPointOffset& OutData, float InputTolerance) const
@@ -467,6 +587,20 @@ bool FNodalOffsetTable::GetPoint(const float InFocus, const float InZoom, FNodal
 
 bool FNodalOffsetTable::SetPoint(float InFocus, float InZoom, const FNodalPointOffset& InData, float InputTolerance)
 {
-	return LensDataTableUtils::SetPoint(*this, InFocus, InZoom, InData, InputTolerance);
+	if (!LensDataTableUtils::SetPoint(*this, InFocus, InZoom, InData, InputTolerance))
+	{
+		return false;
+	}
+
+	LensDataTableUtils::SetPointInFocusCurve(FocusCurves, InFocus, InZoom, InData, InputTolerance);
+	
+	return true;
+}
+
+void FNodalOffsetTable::BuildFocusCurves()
+{
+	// Ensure that the focus curves are empty before building them from the table data
+	FocusCurves.Empty();
+	LensDataTableUtils::BuildFocusCurves(FocusPoints, FocusCurves);
 }
 
diff --git a/Engine/Plugins/VirtualProduction/CameraCalibrationCore/Source/CameraCalibrationCore/Private/Tables/STMapTable.cpp b/Engine/Plugins/VirtualProduction/CameraCalibrationCore/Source/CameraCalibrationCore/Private/Tables/STMapTable.cpp
--- a/Engine/Plugins/VirtualProduction/CameraCalibrationCore/Source/CameraCalibrationCore/Private/Tables/STMapTable.cpp
+++ b/Engine/Plugins/VirtualProduction/CameraCalibrationCore/Source/CameraCalibrationCore/Private/Tables/STMapTable.cpp
@@ -125,6 +125,36 @@ bool FSTMapFocusPoint::IsEmpty() const
 	return MapBlendingCurve.IsEmpty();
 }
 
+void FSTMapFocusCurve::AddPoint(float InFocus, const FSTMapInfo& InData, float InputTolerance)
+{
+	AddPointToCurve(MapBlendingCurve, InFocus, InFocus, InputTolerance);
+}
+
+void FSTMapFocusCurve::SetPoint(float InFocus, const FSTMapInfo& InData, float InputTolerance)
+{
+	// No need to update map curve since x == y
+}
+
+void FSTMapFocusCurve::RemovePoint(float InFocus, float InputTolerance)
+{
+	DeletePointFromCurve(MapBlendingCurve, InFocus, InputTolerance);
+}
+
+void FSTMapFocusCurve::ChangeFocus(float InExistingFocus, float InNewFocus, float InputTolerance)
+{
+	ChangeFocusInCurve(MapBlendingCurve, InExistingFocus, InNewFocus, InputTolerance);
+}
+
+void FSTMapFocusCurve::MergeFocus(float InExistingFocus, float InNewFocus, bool bReplaceExisting, float InputTolerance)
+{
+	MergeFocusInCurve(MapBlendingCurve, InExistingFocus, InNewFocus, bReplaceExisting, InputTolerance);
+}
+
+bool FSTMapFocusCurve::IsEmpty() const
+{
+	return !MapBlendingCurve.GetNumKeys();
+}
+
 void FSTMapTable::ForEachPoint(FFocusPointCallback InCallback) const
 {
 	for (const FSTMapFocusPoint& Point : FocusPoints)
@@ -187,59 +217,28 @@ bool FSTMapTable::BuildParameterCurveAtFocus(float InFocus, int32 InParameterInd
 
 bool FSTMapTable::BuildParameterCurveAtZoom(float InZoom, int32 InParameterIndex, FRichCurve& OutCurve) const
 {
-	for (const FSTMapFocusPoint& FocusPoint : FocusPoints)
+	if (const FSTMapFocusCurve* FocusCurve = GetFocusCurve(InZoom))
 	{
-		FKeyHandle ZoomKeyHandle = FocusPoint.MapBlendingCurve.FindKey(InZoom);
-		if (FocusPoint.MapBlendingCurve.IsKeyHandleValid(ZoomKeyHandle))
-		{
-			const float Value = FocusPoint.MapBlendingCurve.GetKeyValue(ZoomKeyHandle);
-			const FKeyHandle NewKeyHandle = OutCurve.AddKey(FocusPoint.Focus, Value);
-			FRichCurveKey& NewKey = OutCurve.GetKey(NewKeyHandle);
-			NewKey.TangentMode = ERichCurveTangentMode::RCTM_None;
-			NewKey.InterpMode = ERichCurveInterpMode::RCIM_Linear;
-		}
+		OutCurve = FocusCurve->MapBlendingCurve;
+		return true;
 	}
 
-	return true;
+	return false;
 }
 
 void FSTMapTable::SetParameterCurveKeysAtFocus(float InFocus, int32 InParameterIndex, const FRichCurve& InSourceCurve, TArrayView<const FKeyHandle> InKeys)
 {
 	if (FSTMapFocusPoint* FocusPoint = GetFocusPoint(InFocus))
 	{
-		for (int32 Index = 0; Index < InKeys.Num(); ++Index)
-		{
-			const FKeyHandle Handle = InKeys[Index];
-			const int32 KeyIndex = InSourceCurve.GetIndexSafe(Handle);
-			if (KeyIndex != INDEX_NONE)
-			{
-				//We can't move keys on the time axis so our indices should match
-				const FRichCurveKey& Key = InSourceCurve.GetKey(Handle);
-				FocusPoint->MapBlendingCurve.Keys[KeyIndex] = Key;
-			}
-		}
-
-		FocusPoint->MapBlendingCurve.AutoSetTangents();
+		CopyCurveKeys(InSourceCurve, FocusPoint->MapBlendingCurve, InKeys);
 	}
 }
 
 void FSTMapTable::SetParameterCurveKeysAtZoom(float InZoom, int32 InParameterIndex, const FRichCurve& InSourceCurve, TArrayView<const FKeyHandle> InKeys)
 {
-	for (const FKeyHandle& KeyHandle : InKeys)
+	if (FSTMapFocusCurve* FocusCurve = GetFocusCurve(InZoom))
 	{
-		// Assume the focus keys are put into the source curve in the same order as they are stored internally
-		const int32 KeyIndex = InSourceCurve.GetIndexSafe(KeyHandle);
-		if (KeyIndex != INDEX_NONE)
-		{
-			if (ensure(FocusPoints.IsValidIndex(KeyIndex)))
-			{
-				FSTMapFocusPoint& FocusPoint = FocusPoints[KeyIndex];
-				
-				//We can't move keys on the time axis so our indices should match
-				FocusPoint.MapBlendingCurve.Keys[KeyIndex].Value = InSourceCurve.GetKeyValue(KeyHandle);
-				FocusPoint.MapBlendingCurve.AutoSetTangents();
-			}
-		}
+		CopyCurveKeys(InSourceCurve, FocusCurve->MapBlendingCurve, InKeys);
 	}
 }
 
@@ -253,6 +252,16 @@ FSTMapFocusPoint* FSTMapTable::GetFocusPoint(float InFocus, float InputTolerance
 	return FocusPoints.FindByPredicate([InFocus, InputTolerance](const FSTMapFocusPoint& Point) { return FMath::IsNearlyEqual(Point.Focus, InFocus, InputTolerance); });
 }
 
+const FSTMapFocusCurve* FSTMapTable::GetFocusCurve(float InZoom, float InputTolerance) const
+{
+	return FocusCurves.FindByPredicate([InZoom, InputTolerance](const FSTMapFocusCurve& Curve) { return FMath::IsNearlyEqual(Curve.Zoom, InZoom, InputTolerance); });
+}
+
+FSTMapFocusCurve* FSTMapTable::GetFocusCurve(float InZoom, float InputTolerance)
+{
+	return FocusCurves.FindByPredicate([InZoom, InputTolerance](const FSTMapFocusCurve& Curve) { return FMath::IsNearlyEqual(Curve.Zoom, InZoom, InputTolerance); });
+}
+
 TConstArrayView<FSTMapFocusPoint> FSTMapTable::GetFocusPoints() const
 {
 	return FocusPoints;
@@ -263,9 +272,20 @@ TArrayView<FSTMapFocusPoint> FSTMapTable::GetFocusPoints()
 	return FocusPoints;
 }
 
+TConstArrayView<FSTMapFocusCurve> FSTMapTable::GetFocusCurves() const
+{
+	return FocusCurves;
+}
+
+TArray<FSTMapFocusCurve>& FSTMapTable::GetFocusCurves()
+{
+	return FocusCurves;
+}
+
 void FSTMapTable::RemoveFocusPoint(float InFocus)
 {
 	LensDataTableUtils::RemoveFocusPoint(FocusPoints, InFocus);
+	LensDataTableUtils::RemoveFocusFromFocusCurves(FocusCurves, InFocus);
 }
 
 bool FSTMapTable::HasFocusPoint(float InFocus, float InputTolerance) const
@@ -276,16 +296,19 @@ bool FSTMapTable::HasFocusPoint(float InFocus, float InputTolerance) const
 void FSTMapTable::ChangeFocusPoint(float InExistingFocus, float InNewFocus, float InputTolerance)
 {
 	LensDataTableUtils::ChangeFocusPoint(FocusPoints, InExistingFocus, InNewFocus, InputTolerance);
+	LensDataTableUtils::ChangeFocusInFocusCurves(FocusCurves, InExistingFocus, InNewFocus, InputTolerance);
 }
 
 void FSTMapTable::MergeFocusPoint(float InSrcFocus, float InDestFocus, bool bReplaceExistingZoomPoints, float InputTolerance)
 {
 	LensDataTableUtils::MergeFocusPoint(FocusPoints, InSrcFocus, InDestFocus, bReplaceExistingZoomPoints, InputTolerance);
+	LensDataTableUtils::MergeFocusInFocusCurves(FocusCurves, InSrcFocus, InDestFocus, bReplaceExistingZoomPoints, InputTolerance);
 }
 
 void FSTMapTable::RemoveZoomPoint(float InFocus, float InZoom)
 {
 	LensDataTableUtils::RemoveZoomPoint(FocusPoints, InFocus, InZoom);
+	LensDataTableUtils::RemoveZoomFromFocusCurves(FocusCurves, InFocus, InZoom);
 }
 
 bool FSTMapTable::HasZoomPoint(float InFocus, float InZoom, float InputTolerance)
@@ -296,6 +319,14 @@ bool FSTMapTable::HasZoomPoint(float InFocus, float InZoom, float InputTolerance
 void FSTMapTable::ChangeZoomPoint(float InFocus, float InExistingZoom, float InNewZoom, float InputTolerance)
 {
 	LensDataTableUtils::ChangeZoomPoint(FocusPoints, InFocus, InExistingZoom, InNewZoom, InputTolerance);
+	
+	FSTMapInfo Data;
+	if (!GetPoint(InFocus, InNewZoom, Data, InputTolerance))
+	{
+		return;
+	}
+
+	LensDataTableUtils::ChangeZoomInFocusCurves(FocusCurves, InFocus, InExistingZoom, InNewZoom, Data, InputTolerance);
 }
 
 TMap<ELensDataCategory, FLinkPointMetadata> FSTMapTable::GetLinkedCategories() const
@@ -312,7 +343,13 @@ TMap<ELensDataCategory, FLinkPointMetadata> FSTMapTable::GetLinkedCategories() c
 bool FSTMapTable::AddPoint(float InFocus, float InZoom, const FSTMapInfo& InData, float InputTolerance,
                            bool bIsCalibrationPoint)
 {
-	return LensDataTableUtils::AddPoint(FocusPoints, InFocus, InZoom, InData, InputTolerance, bIsCalibrationPoint);
+	if (!LensDataTableUtils::AddPoint(FocusPoints, InFocus, InZoom, InData, InputTolerance, bIsCalibrationPoint))
+	{
+		return false;
+	}
+
+	LensDataTableUtils::AddPointToFocusCurve(FocusCurves, InFocus, InZoom, InData, InputTolerance);
+	return true;
 }
 
 bool FSTMapTable::GetPoint(const float InFocus, const float InZoom, FSTMapInfo& OutData, float InputTolerance) const
@@ -333,6 +370,20 @@ bool FSTMapTable::GetPoint(const float InFocus, const float InZoom, FSTMapInfo&
 
 bool FSTMapTable::SetPoint(float InFocus, float InZoom, const FSTMapInfo& InData, float InputTolerance)
 {
-	return LensDataTableUtils::SetPoint(*this, InFocus, InZoom, InData, InputTolerance);
+	if (!LensDataTableUtils::SetPoint(*this, InFocus, InZoom, InData, InputTolerance))
+	{
+		return false;
+	}
+
+	LensDataTableUtils::SetPointInFocusCurve(FocusCurves, InFocus, InZoom, InData, InputTolerance);
+	
+	return true;
+}
+
+void FSTMapTable::BuildFocusCurves()
+{
+	// Ensure that the focus curves are empty before building them from the table data
+	FocusCurves.Empty();
+	LensDataTableUtils::BuildFocusCurves(FocusPoints, FocusCurves);
 }
 
diff --git a/Engine/Plugins/VirtualProduction/CameraCalibrationCore/Source/CameraCalibrationCore/Public/LensFile.h b/Engine/Plugins/VirtualProduction/CameraCalibrationCore/Source/CameraCalibrationCore/Public/LensFile.h
--- a/Engine/Plugins/VirtualProduction/CameraCalibrationCore/Source/CameraCalibrationCore/Public/LensFile.h
+++ b/Engine/Plugins/VirtualProduction/CameraCalibrationCore/Source/CameraCalibrationCore/Public/LensFile.h
@@ -68,6 +68,8 @@ public:
 	ULensFile();
 
 	//~Begin UObject interface
+	virtual void Serialize(FArchive& Ar) override;
+	
 #if WITH_EDITOR
 	virtual void PostEditChangeChainProperty( struct FPropertyChangedChainEvent& PropertyChangedEvent ) override;
 #endif //WITH_EDITOR
@@ -257,6 +259,11 @@ protected:
 	
 	/** Evaluates distortion based on InFocus and InZoom using STMaps */
 	bool EvaluateDistortionForSTMaps(float InFocus, float InZoom, FVector2D InFilmback, ULensDistortionModelHandlerBase* LensHandler) const;
+
+#if WITH_EDITOR
+	/** Builds the lens table focus curves to match the existing data in the tables */
+	void BuildLensTableFocusCurves();
+#endif // WITH_EDITOR
 	
 public:
 
diff --git a/Engine/Plugins/VirtualProduction/CameraCalibrationCore/Source/CameraCalibrationCore/Public/LensFileObjectVersion.h b/Engine/Plugins/VirtualProduction/CameraCalibrationCore/Source/CameraCalibrationCore/Public/LensFileObjectVersion.h
new file mode 100644
--- /dev/null
+++ b/Engine/Plugins/VirtualProduction/CameraCalibrationCore/Source/CameraCalibrationCore/Public/LensFileObjectVersion.h
@@ -0,0 +1,22 @@
+﻿// Copyright Epic Games, Inc. All Rights Reserved.
+#pragma once
+
+#include "CoreTypes.h"
+#include "Misc/Guid.h"
+
+struct CAMERACALIBRATIONCORE_API FLensFileObjectVersion
+{
+	enum Type
+	{
+		BeforeCustomVersionWasAdded = 0,
+		EditableFocusCurves,
+
+		// -----<new versions can be added above this line>-------------------------------------------------
+		VersionPlusOne,
+		LatestVersion = VersionPlusOne - 1
+	};
+	
+	const static FGuid GUID;
+	
+	FLensFileObjectVersion() = delete;
+};
\ No newline at end of file
diff --git a/Engine/Plugins/VirtualProduction/CameraCalibrationCore/Source/CameraCalibrationCore/Public/Tables/BaseLensTable.h b/Engine/Plugins/VirtualProduction/CameraCalibrationCore/Source/CameraCalibrationCore/Public/Tables/BaseLensTable.h
--- a/Engine/Plugins/VirtualProduction/CameraCalibrationCore/Source/CameraCalibrationCore/Public/Tables/BaseLensTable.h
+++ b/Engine/Plugins/VirtualProduction/CameraCalibrationCore/Source/CameraCalibrationCore/Public/Tables/BaseLensTable.h
@@ -3,6 +3,8 @@
 #pragma once
 
 #include "CoreMinimal.h"
+#include "Curves/KeyHandle.h"
+#include "Curves/RichCurve.h"
 
 #include "UObject/WeakObjectPtrTemplates.h"
 
@@ -58,6 +60,29 @@ struct TStructOpsTypeTraits<FBaseFocusPoint> : public TStructOpsTypeTraitsBase2<
 	};
 };
 
+/** Base focus curve struct */
+USTRUCT()
+struct FBaseFocusCurve
+{
+	GENERATED_BODY()
+
+protected:
+	/** Adds a new key to the specified curve */
+	FKeyHandle AddPointToCurve(FRichCurve& InCurve, float InFocus, float InValue, float InputTolerance, FKeyHandle InOptionalKeyHandle = FKeyHandle());
+	
+	/** Sets the value of an existing key in the specified curve */
+	FKeyHandle SetPointInCurve(FRichCurve& InCurve, float InFocus, float InValue, float InputTolerance);
+	
+	/** Deletes a key at the specified focus from the specified curve */
+	void DeletePointFromCurve(FRichCurve& InCurve, float InFocus, float InputTolerance);
+	
+	/** Changes the focus of a key in the specified curve */
+	void ChangeFocusInCurve(FRichCurve& InCurve, float InExistingFocus, float InNewFocus, float InputTolerance);
+	
+	/** Changes the focus of a key in the specified curve and optionally replaces any key that already exists at the new focus */
+	void MergeFocusInCurve(FRichCurve& InCurve, float InExistingFocus, float InNewFocus, bool bReplaceExisting, float InputTolerance);
+};
+
 /**
  * Base data table struct
  */
@@ -95,6 +120,12 @@ protected:
 	 */
 	virtual bool DoesZoomPointExists(float InFocus, float InZoom, float InputTolerance = KINDA_SMALL_NUMBER) const PURE_VIRTUAL(FBaseLensTable::DoesZoomPointExists, return false; );
 
+	/** Copies the specified keys from the source curve to the destination curve */
+	void CopyCurveKeys(const FRichCurve& InSourceCurve, FRichCurve& InDestCurve, TArrayView<const FKeyHandle> InKeys);
+
+	/** Propagates the values of a curve to a set of cross curves at the specified time */
+	void PropagateCurveValuesToCrossCurves(const FRichCurve& InCurve, float InCrossCurveTime, TFunctionRef<FRichCurve*(float)> GetCurveFn);
+	
 public:
 	virtual ~FBaseLensTable() = default;
 
diff --git a/Engine/Plugins/VirtualProduction/CameraCalibrationCore/Source/CameraCalibrationCore/Public/Tables/DistortionParametersTable.h b/Engine/Plugins/VirtualProduction/CameraCalibrationCore/Source/CameraCalibrationCore/Public/Tables/DistortionParametersTable.h
--- a/Engine/Plugins/VirtualProduction/CameraCalibrationCore/Source/CameraCalibrationCore/Public/Tables/DistortionParametersTable.h
+++ b/Engine/Plugins/VirtualProduction/CameraCalibrationCore/Source/CameraCalibrationCore/Public/Tables/DistortionParametersTable.h
@@ -61,6 +61,9 @@ public:
 	
 	/** Returns true if this point is empty */
 	bool IsEmpty() const;
+
+	/** Gets the curve for the specified parameter, or nullptr if the parameter index is invalid */
+	const FRichCurve* GetCurveForParameter(int32 InParameterIndex) const;
 	
 	void SetParameterValue(int32 InZoomIndex, float InZoomValue, int32 InParameterIndex, float InParameterValue);
 
@@ -79,6 +82,44 @@ public:
 	TArray<FDistortionZoomPoint> ZoomPoints;
 };
 
+/** A curve along the focus axis for a single zoom value */
+USTRUCT()
+struct CAMERACALIBRATIONCORE_API FDistortionFocusCurve : public FBaseFocusCurve
+{
+	GENERATED_BODY()
+
+public:
+	/** Adds a new point to the focus curve, or updates a matching existing point if one is found */
+	void AddPoint(float InFocus, const FDistortionInfo& InData, float InputTolerance = KINDA_SMALL_NUMBER);
+
+	/** Updates an existing point if one is found */
+	void SetPoint(float InFocus, const FDistortionInfo& InData, float InputTolerance = KINDA_SMALL_NUMBER);
+
+	/** Removes the point at the specified focus if one is found */
+	void RemovePoint(float InFocus, float InputTolerance = KINDA_SMALL_NUMBER);
+
+	/** Changes the focus value of the point at the specified focus, if one is found */
+	void ChangeFocus(float InExistingFocus, float InNewFocus, float InputTolerance = KINDA_SMALL_NUMBER);
+
+	/** Changes the focus value of the point at the specified focus and optionally replaces any point at the new focus with the old point */
+	void MergeFocus(float InExistingFocus, float InNewFocus, bool bReplaceExisting, float InputTolerance = KINDA_SMALL_NUMBER);
+
+	/** Gets whether the curve is empty */
+	bool IsEmpty() const;
+
+	/** Gets the curve for the specified parameter, or nullptr if the parameter index is invalid */
+	const FRichCurve* GetCurveForParameter(int32 InParameterIndex) const;
+	
+public:
+	/** Curve describing desired blending between resulting displacement maps */
+	UPROPERTY()
+	FRichCurve MapBlendingCurve;
+
+	/** The fixed zoom value of the curve */
+	UPROPERTY()
+	float Zoom = 0.0f;
+};
+
 /**
  * Distortion table containing list of points for each focus and zoom input
  */
@@ -88,6 +129,7 @@ struct CAMERACALIBRATIONCORE_API FDistortionTable : public FBaseLensTable
 	GENERATED_BODY()
 
 	using FocusPointType = FDistortionFocusPoint;
+	using FocusCurveType = FDistortionFocusCurve;
 
 	/** Wrapper for indices of specific parameters for the distortion table  */
 	struct FParameters
@@ -124,11 +166,23 @@ public:
 	/** Returns point for a given focus */
 	FDistortionFocusPoint* GetFocusPoint(float InFocus, float InputTolerance = KINDA_SMALL_NUMBER);
 
+	/** Gets the focus curve for the specified zoom, or nullptr if none were found */
+	const FDistortionFocusCurve* GetFocusCurve(float InZoom, float InputTolerance = KINDA_SMALL_NUMBER) const;
+
+	/** Gets the focus curve for the specified zoom, or nullptr if none were found */
+	FDistortionFocusCurve* GetFocusCurve(float InZoom, float InputTolerance = KINDA_SMALL_NUMBER);
+	
 	/** Returns all focus points */
 	TConstArrayView<FDistortionFocusPoint> GetFocusPoints() const;
 
 	/** Returns all focus points */
 	TArray<FDistortionFocusPoint>& GetFocusPoints();
+
+	/** Returns all focus curves */
+	TConstArrayView<FDistortionFocusCurve> GetFocusCurves() const;
+
+	/** Returns all focus curves */
+	TArray<FDistortionFocusCurve>& GetFocusCurves();
 	
 	/** Removes a focus point identified as InFocusIdentifier */
 	void RemoveFocusPoint(float InFocus);
@@ -160,10 +214,17 @@ public:
 	/** Set a new point into the table */
 	bool SetPoint(float InFocus, float InZoom, const FDistortionInfo& InData, float InputTolerance = KINDA_SMALL_NUMBER);
 
+	/** Builds the focus curves to match existing data in the table */
+	void BuildFocusCurves();
+	
 public:
 
 	/** Lists of focus points */
 	UPROPERTY()
 	TArray<FDistortionFocusPoint> FocusPoints;
+
+	/** A list of curves along the focus axis for each zoom value */
+	UPROPERTY()
+	TArray<FDistortionFocusCurve> FocusCurves;
 };
 
diff --git a/Engine/Plugins/VirtualProduction/CameraCalibrationCore/Source/CameraCalibrationCore/Public/Tables/FocalLengthTable.h b/Engine/Plugins/VirtualProduction/CameraCalibrationCore/Source/CameraCalibrationCore/Public/Tables/FocalLengthTable.h
--- a/Engine/Plugins/VirtualProduction/CameraCalibrationCore/Source/CameraCalibrationCore/Public/Tables/FocalLengthTable.h
+++ b/Engine/Plugins/VirtualProduction/CameraCalibrationCore/Source/CameraCalibrationCore/Public/Tables/FocalLengthTable.h
@@ -70,6 +70,12 @@ public:
 	/** Returns true if this point is empty */
 	bool IsEmpty() const;
 
+	/** Gets the curve for the specified parameter, or nullptr if the parameter index is invalid */
+	const FRichCurve* GetCurveForParameter(int32 InParameterIndex) const;
+
+	/** Gets the curve for the specified parameter, or nullptr if the parameter index is invalid */
+	FRichCurve* GetCurveForParameter(int32 InParameterIndex);
+	
 public:
 
 	/** Input focus for this point */
@@ -89,6 +95,50 @@ public:
 	TArray<FFocalLengthZoomPoint> ZoomPoints;
 };
 
+USTRUCT()
+struct CAMERACALIBRATIONCORE_API FFocalLengthFocusCurve : public FBaseFocusCurve
+{
+	GENERATED_BODY()
+
+public:
+	/** Adds a new point to the focus curve, or updates a matching existing point if one is found */
+	void AddPoint(float InFocus, const FFocalLengthInfo& InData, float InputTolerance = KINDA_SMALL_NUMBER);
+
+	/** Updates an existing point if one is found */
+	void SetPoint(float InFocus, const FFocalLengthInfo& InData, float InputTolerance = KINDA_SMALL_NUMBER);
+
+	/** Removes the point at the specified focus if one is found */
+	void RemovePoint(float InFocus, float InputTolerance = KINDA_SMALL_NUMBER);
+
+	/** Changes the focus value of the point at the specified focus, if one is found */
+	void ChangeFocus(float InExistingFocus, float InNewFocus, float InputTolerance = KINDA_SMALL_NUMBER);
+
+	/** Changes the focus value of the point at the specified focus and optionally replaces any point at the new focus with the old point */
+	void MergeFocus(float InExistingFocus, float InNewFocus, bool bReplaceExisting, float InputTolerance = KINDA_SMALL_NUMBER);
+
+	/** Gets whether the curve is empty */
+	bool IsEmpty() const;
+
+	/** Gets the curve for the specified parameter, or nullptr if the parameter index is invalid */
+	const FRichCurve* GetCurveForParameter(int32 InParameterIndex) const;
+
+	/** Gets the curve for the specified parameter, or nullptr if the parameter index is invalid */
+	FRichCurve* GetCurveForParameter(int32 InParameterIndex);
+	
+public:
+	/** Focus curve for the x parameter of the focal length */
+	UPROPERTY()
+	FRichCurve Fx;
+
+	/** Focus curve for the y parameter of the focal length */
+	UPROPERTY()
+	FRichCurve Fy;
+
+	/** The fixed zoom value of the curve */
+	UPROPERTY()
+	float Zoom = 0.0f;
+};
+
 /**
  * Focal Length table containing FxFy values for each focus and zoom input values
  */
@@ -98,6 +148,7 @@ struct CAMERACALIBRATIONCORE_API FFocalLengthTable : public FBaseLensTable
 	GENERATED_BODY()
 
 	using FocusPointType = FFocalLengthFocusPoint;
+	using FocusCurveType = FFocalLengthFocusCurve;
 
 	/** Wrapper for indices of specific parameters for the focal length table  */
 	struct FParameters
@@ -144,8 +195,23 @@ public:
 	/** Returns point for a given focus */
 	FFocalLengthFocusPoint* GetFocusPoint(float InFocus, float InputTolerance = KINDA_SMALL_NUMBER);
 
+	/** Gets the focus curve for the specified zoom, or nullptr if none were found */
+	const FFocalLengthFocusCurve* GetFocusCurve(float InZoom, float InputTolerance = KINDA_SMALL_NUMBER) const;
+
+	/** Gets the focus curve for the specified zoom, or nullptr if none were found */
+	FFocalLengthFocusCurve* GetFocusCurve(float InZoom, float InputTolerance = KINDA_SMALL_NUMBER);
+	
 	/** Returns all focus points */
 	TConstArrayView<FFocalLengthFocusPoint> GetFocusPoints() const;
+
+	/** Returns all focus points */
+	TArray<FFocalLengthFocusPoint>& GetFocusPoints();
+	
+	/** Returns all focus curves */
+	TConstArrayView<FFocalLengthFocusCurve> GetFocusCurves() const;
+
+	/** Returns all focus curves */
+	TArray<FFocalLengthFocusCurve>& GetFocusCurves();
 	
 	/** Removes a focus point identified as InFocusIdentifier */
 	void RemoveFocusPoint(float InFocus);
@@ -176,11 +242,18 @@ public:
 
 	/** Set a new point into the table */
 	bool SetPoint(float InFocus, float InZoom, const FFocalLengthInfo& InData, float InputTolerance = KINDA_SMALL_NUMBER);
-
+	
+	/** Builds the focus curves to match existing data in the table */
+	void BuildFocusCurves();
+	
 public:
 
 	/** Lists of focus points */
 	UPROPERTY()
 	TArray<FFocalLengthFocusPoint> FocusPoints;
+
+	/** A list of curves along the focus axis for each zoom value */
+	UPROPERTY()
+	TArray<FFocalLengthFocusCurve> FocusCurves;
 };
 
diff --git a/Engine/Plugins/VirtualProduction/CameraCalibrationCore/Source/CameraCalibrationCore/Public/Tables/ImageCenterTable.h b/Engine/Plugins/VirtualProduction/CameraCalibrationCore/Source/CameraCalibrationCore/Public/Tables/ImageCenterTable.h
--- a/Engine/Plugins/VirtualProduction/CameraCalibrationCore/Source/CameraCalibrationCore/Public/Tables/ImageCenterTable.h
+++ b/Engine/Plugins/VirtualProduction/CameraCalibrationCore/Source/CameraCalibrationCore/Public/Tables/ImageCenterTable.h
@@ -8,7 +8,6 @@
 
 #include "ImageCenterTable.generated.h"
 
-
 /**
  * ImageCenter focus point containing curves for CxCy 
  */
@@ -44,6 +43,12 @@ public:
 	/** Returns true if this point is empty */
 	bool IsEmpty() const;
 
+	/** Gets the curve for the specified parameter, or nullptr if the parameter index is invalid */
+	const FRichCurve* GetCurveForParameter(int32 InParameterIndex) const;
+
+	/** Gets the curve for the specified parameter, or nullptr if the parameter index is invalid */
+	FRichCurve* GetCurveForParameter(int32 InParameterIndex);
+	
 public:
 
 	/** Focus value of this point */
@@ -59,6 +64,51 @@ public:
 	FRichCurve Cy;
 };
 
+/** A curve along the focus axis for a single zoom value */
+USTRUCT()
+struct CAMERACALIBRATIONCORE_API FImageCenterFocusCurve : public FBaseFocusCurve
+{
+	GENERATED_BODY()
+
+public:
+	/** Adds a new point to the focus curve, or updates a matching existing point if one is found */
+	void AddPoint(float InFocus, const FImageCenterInfo& InData, float InputTolerance = KINDA_SMALL_NUMBER);
+
+	/** Updates an existing point if one is found */
+	void SetPoint(float InFocus, const FImageCenterInfo& InData, float InputTolerance = KINDA_SMALL_NUMBER);
+
+	/** Removes the point at the specified focus if one is found */
+	void RemovePoint(float InFocus, float InputTolerance = KINDA_SMALL_NUMBER);
+
+	/** Changes the focus value of the point at the specified focus, if one is found */
+	void ChangeFocus(float InExistingFocus, float InNewFocus, float InputTolerance = KINDA_SMALL_NUMBER);
+
+	/** Changes the focus value of the point at the specified focus and optionally replaces any point at the new focus with the old point */
+	void MergeFocus(float InExistingFocus, float InNewFocus, bool bReplaceExisting, float InputTolerance = KINDA_SMALL_NUMBER);
+
+	/** Gets whether the curve is empty */
+	bool IsEmpty() const;
+
+	/** Gets the curve for the specified parameter, or nullptr if the parameter index is invalid */
+	const FRichCurve* GetCurveForParameter(int32 InParameterIndex) const;
+
+	/** Gets the curve for the specified parameter, or nullptr if the parameter index is invalid */
+	FRichCurve* GetCurveForParameter(int32 InParameterIndex);
+	
+public:
+	/** Focus curve for the x parameter of the image center */
+	UPROPERTY()
+	FRichCurve Cx;
+
+	/** Focus curve for the y parameter of the image center */
+	UPROPERTY()
+	FRichCurve Cy;
+
+	/** The fixed zoom value of the curve */
+	UPROPERTY()
+	float Zoom = 0.0f;
+};
+
 /**
  * Image Center table associating CxCy values to focus and zoom
  */
@@ -68,7 +118,8 @@ struct CAMERACALIBRATIONCORE_API FImageCenterTable : public FBaseLensTable
 	GENERATED_BODY()
 
 	using FocusPointType = FImageCenterFocusPoint;
-
+	using FocusCurveType = FImageCenterFocusCurve;
+	
 	/** Wrapper for indices of specific parameters for the image center table  */
 	struct FParameters
 	{
@@ -108,11 +159,23 @@ public:
 	/** Returns const point for a given focus */
 	FImageCenterFocusPoint* GetFocusPoint(float InFocus, float InputTolerance = KINDA_SMALL_NUMBER);
 
+	/** Gets the focus curve for the specified zoom, or nullptr if none were found */
+	const FImageCenterFocusCurve* GetFocusCurve(float InZoom, float InputTolerance = KINDA_SMALL_NUMBER) const;
+
+	/** Gets the focus curve for the specified zoom, or nullptr if none were found */
+	FImageCenterFocusCurve* GetFocusCurve(float InZoom, float InputTolerance = KINDA_SMALL_NUMBER);
+	
 	/** Returns all focus points */
 	TConstArrayView<FImageCenterFocusPoint> GetFocusPoints() const;
 
 	/** Returns all focus points */
 	TArray<FImageCenterFocusPoint>& GetFocusPoints();
+
+	/** Returns all focus curves */
+	TConstArrayView<FImageCenterFocusCurve> GetFocusCurves() const;
+
+	/** Returns all focus curves */
+	TArray<FImageCenterFocusCurve>& GetFocusCurves();
 	
 	/** Removes a focus point identified as InFocusIdentifier */
 	void RemoveFocusPoint(float InFocus);
@@ -144,10 +207,17 @@ public:
 	/** Set a new point into the table */
 	bool SetPoint(float InFocus, float InZoom, const FImageCenterInfo& InData, float InputTolerance = KINDA_SMALL_NUMBER);
 
-public:
+	/** Builds the focus curves to match existing data in the table */
+	void BuildFocusCurves();
+	
+public:		
 
 	/** Lists of focus points */
 	UPROPERTY()
 	TArray<FImageCenterFocusPoint> FocusPoints;
+
+	/** A list of curves along the focus axis for each zoom value */
+	UPROPERTY()
+	TArray<FImageCenterFocusCurve> FocusCurves;
 };
 
diff --git a/Engine/Plugins/VirtualProduction/CameraCalibrationCore/Source/CameraCalibrationCore/Public/Tables/NodalOffsetTable.h b/Engine/Plugins/VirtualProduction/CameraCalibrationCore/Source/CameraCalibrationCore/Public/Tables/NodalOffsetTable.h
--- a/Engine/Plugins/VirtualProduction/CameraCalibrationCore/Source/CameraCalibrationCore/Public/Tables/NodalOffsetTable.h
+++ b/Engine/Plugins/VirtualProduction/CameraCalibrationCore/Source/CameraCalibrationCore/Public/Tables/NodalOffsetTable.h
@@ -44,6 +44,12 @@ public:
 	/** Returns true if there are no points */
 	bool IsEmpty() const;
 
+	/** Gets the curve for the specified parameter, or nullptr if the parameter index is invalid */
+	const FRichCurve* GetCurveForParameter(int32 InParameterIndex) const;
+
+	/** Gets the curve for the specified parameter, or nullptr if the parameter index is invalid */
+	FRichCurve* GetCurveForParameter(int32 InParameterIndex);
+	
 public:
 
 	/** Dimensions of our location offset curves */
@@ -65,6 +71,51 @@ public:
 	FRichCurve RotationOffset[RotationDimension];
 };
 
+/** A curve along the focus axis for a single zoom value */
+USTRUCT()
+struct CAMERACALIBRATIONCORE_API FNodalOffsetFocusCurve : public FBaseFocusCurve
+{
+	GENERATED_BODY()
+
+public:
+	/** Adds a new point to the focus curve, or updates a matching existing point if one is found */
+	void AddPoint(float InFocus, const FNodalPointOffset& InData, float InputTolerance = KINDA_SMALL_NUMBER);
+
+	/** Updates an existing point if one is found */
+	void SetPoint(float InFocus, const FNodalPointOffset& InData, float InputTolerance = KINDA_SMALL_NUMBER);
+
+	/** Removes the point at the specified focus if one is found */
+	void RemovePoint(float InFocus, float InputTolerance = KINDA_SMALL_NUMBER);
+
+	/** Changes the focus value of the point at the specified focus, if one is found */
+	void ChangeFocus(float InExistingFocus, float InNewFocus, float InputTolerance = KINDA_SMALL_NUMBER);
+
+	/** Changes the focus value of the point at the specified focus and optionally replaces any point at the new focus with the old point */
+	void MergeFocus(float InExistingFocus, float InNewFocus, bool bReplaceExisting, float InputTolerance = KINDA_SMALL_NUMBER);
+
+	/** Gets whether the curve is empty */
+	bool IsEmpty() const;
+
+	/** Gets the curve for the specified parameter, or nullptr if the parameter index is invalid */
+	const FRichCurve* GetCurveForParameter(int32 InParameterIndex) const;
+
+	/** Gets the curve for the specified parameter, or nullptr if the parameter index is invalid */
+	FRichCurve* GetCurveForParameter(int32 InParameterIndex);
+	
+public:
+	/** Focus curve for the location parameters of the nodal offset */
+	UPROPERTY()
+	FRichCurve LocationOffset[FNodalOffsetFocusPoint::LocationDimension];
+
+	/** Focus curve for the rotation parameters of the nodal offset */
+	UPROPERTY()
+	FRichCurve RotationOffset[FNodalOffsetFocusPoint::RotationDimension];
+	
+	/** The fixed zoom value of the curve */
+	UPROPERTY()
+	float Zoom = 0.0f;
+};
+
 /**
  * Table containing nodal offset mapping to focus and zoom
  */
@@ -74,6 +125,7 @@ struct CAMERACALIBRATIONCORE_API FNodalOffsetTable : public FBaseLensTable
 	GENERATED_BODY()
 
 	using FocusPointType = FNodalOffsetFocusPoint;
+	using FocusCurveType = FNodalOffsetFocusCurve;
 
 	/** Wrapper for indices of specific parameters for the nodal offset table  */
 	struct FParameters
@@ -125,12 +177,24 @@ public:
 	/** Returns point for a given focus */
 	FNodalOffsetFocusPoint* GetFocusPoint(float InFocus, float InputTolerance = KINDA_SMALL_NUMBER);
 
+	/** Gets the focus curve for the specified zoom, or nullptr if none were found */
+	const FNodalOffsetFocusCurve* GetFocusCurve(float InZoom, float InputTolerance = KINDA_SMALL_NUMBER) const;
+
+	/** Gets the focus curve for the specified zoom, or nullptr if none were found */
+	FNodalOffsetFocusCurve* GetFocusCurve(float InZoom, float InputTolerance = KINDA_SMALL_NUMBER);
+	
 	/** Returns all focus points */
 	TConstArrayView<FNodalOffsetFocusPoint> GetFocusPoints() const;
 
 	/** Returns all focus points */
 	TArray<FNodalOffsetFocusPoint>& GetFocusPoints();
 
+	/** Returns all focus curves */
+	TConstArrayView<FNodalOffsetFocusCurve> GetFocusCurves() const;
+
+	/** Returns all focus curves */
+	TArray<FNodalOffsetFocusCurve>& GetFocusCurves();
+	
 	/** Removes a focus point identified as InFocusIdentifier */
 	void RemoveFocusPoint(float InFocus);
 
@@ -160,11 +224,18 @@ public:
 
 	/** Set a new point into the table */
 	bool SetPoint(float InFocus, float InZoom, const FNodalPointOffset& InData, float InputTolerance = KINDA_SMALL_NUMBER);
+
+	/** Builds the focus curves to match existing data in the table */
+	void BuildFocusCurves();
 	
 public:
 
 	/** Lists of focus points */
 	UPROPERTY()
 	TArray<FNodalOffsetFocusPoint> FocusPoints;
+
+	/** A list of curves along the focus axis for each zoom value */
+	UPROPERTY()
+	TArray<FNodalOffsetFocusCurve> FocusCurves;
 };
 
diff --git a/Engine/Plugins/VirtualProduction/CameraCalibrationCore/Source/CameraCalibrationCore/Public/Tables/STMapTable.h b/Engine/Plugins/VirtualProduction/CameraCalibrationCore/Source/CameraCalibrationCore/Public/Tables/STMapTable.h
--- a/Engine/Plugins/VirtualProduction/CameraCalibrationCore/Source/CameraCalibrationCore/Public/Tables/STMapTable.h
+++ b/Engine/Plugins/VirtualProduction/CameraCalibrationCore/Source/CameraCalibrationCore/Public/Tables/STMapTable.h
@@ -117,6 +117,41 @@ public:
 	TArray<FSTMapZoomPoint> ZoomPoints;
 };
 
+/** A curve along the focus axis for a single zoom value */
+USTRUCT()
+struct CAMERACALIBRATIONCORE_API FSTMapFocusCurve : public FBaseFocusCurve
+{
+	GENERATED_BODY()
+
+public:
+	/** Adds a new point to the focus curve, or updates a matching existing point if one is found */
+	void AddPoint(float InFocus, const FSTMapInfo& InData, float InputTolerance = KINDA_SMALL_NUMBER);
+
+	/** Updates an existing point if one is found */
+	void SetPoint(float InFocus, const FSTMapInfo& InData, float InputTolerance = KINDA_SMALL_NUMBER);
+
+	/** Removes the point at the specified focus if one is found */
+	void RemovePoint(float InFocus, float InputTolerance = KINDA_SMALL_NUMBER);
+
+	/** Changes the focus value of the point at the specified focus, if one is found */
+	void ChangeFocus(float InExistingFocus, float InNewFocus, float InputTolerance = KINDA_SMALL_NUMBER);
+
+	/** Changes the focus value of the point at the specified focus and optionally replaces any point at the new focus with the old point */
+	void MergeFocus(float InExistingFocus, float InNewFocus, bool bReplaceExisting, float InputTolerance = KINDA_SMALL_NUMBER);
+
+	/** Gets whether the curve is empty */
+	bool IsEmpty() const;
+
+public:
+	/** Curve describing desired blending between resulting displacement maps */
+	UPROPERTY()
+	FRichCurve MapBlendingCurve;
+
+	/** The fixed zoom value of the curve */
+	UPROPERTY()
+	float Zoom = 0.0f;
+};
+
 /**
  * STMap table containing list of points for each focus and zoom inputs
  */
@@ -126,6 +161,7 @@ struct CAMERACALIBRATIONCORE_API FSTMapTable : public FBaseLensTable
 	GENERATED_BODY()
 
 	using FocusPointType = FSTMapFocusPoint;
+	using FocusCurveType = FSTMapFocusCurve;
 
 protected:
 	//~ Begin FBaseDataTable Interface
@@ -155,12 +191,24 @@ public:
 	/** Returns point for a given focus */
 	FSTMapFocusPoint* GetFocusPoint(float InFocus, float InputTolerance = KINDA_SMALL_NUMBER);
 
+	/** Gets the focus curve for the specified zoom, or nullptr if none were found */
+	const FSTMapFocusCurve* GetFocusCurve(float InZoom, float InputTolerance = KINDA_SMALL_NUMBER) const;
+
+	/** Gets the focus curve for the specified zoom, or nullptr if none were found */
+	FSTMapFocusCurve* GetFocusCurve(float InZoom, float InputTolerance = KINDA_SMALL_NUMBER);
+	
 	/** Returns all focus points */
 	TConstArrayView<FSTMapFocusPoint> GetFocusPoints() const;
 
 	/** Returns all focus points */
 	TArrayView<FSTMapFocusPoint> GetFocusPoints();
 
+	/** Returns all focus curves */
+	TConstArrayView<FSTMapFocusCurve> GetFocusCurves() const;
+
+	/** Returns all focus curves */
+	TArray<FSTMapFocusCurve>& GetFocusCurves();
+	
 	/** Removes a focus point identified as InFocusIdentifier */
 	void RemoveFocusPoint(float InFocus);
 
@@ -191,10 +239,17 @@ public:
 	/** Set a new point into the table */
 	bool SetPoint(float InFocus, float InZoom, const FSTMapInfo& InData, float InputTolerance = KINDA_SMALL_NUMBER);
 
+	/** Builds the focus curves to match existing data in the table */
+	void BuildFocusCurves();
+	
 public:
 
 	/** Lists of focus points */
 	UPROPERTY()
 	TArray<FSTMapFocusPoint> FocusPoints;
+
+	/** A list of curves along the focus axis for each zoom value */
+	UPROPERTY()
+	TArray<FSTMapFocusCurve> FocusCurves;
 };
 
