CameraCalibration: Expose checkerboard actor properties to BP

diff --git a/Engine/Plugins/VirtualProduction/CameraCalibrationCore/Source/CameraCalibrationCore/Public/CameraCalibrationCheckerboard.h b/Engine/Plugins/VirtualProduction/CameraCalibrationCore/Source/CameraCalibrationCore/Public/CameraCalibrationCheckerboard.h
--- a/Engine/Plugins/VirtualProduction/CameraCalibrationCore/Source/CameraCalibrationCore/Public/CameraCalibrationCheckerboard.h
+++ b/Engine/Plugins/VirtualProduction/CameraCalibrationCore/Source/CameraCalibrationCore/Public/CameraCalibrationCheckerboard.h
@@ -69,19 +69,19 @@ public:
 
 
 	/** Number of rows */
-	UPROPERTY(EditAnywhere, Category = "Calibration")
+	UPROPERTY(EditAnywhere, BlueprintReadWrite, Category = "Calibration")
 	int32 NumCornerRows = 2;
 
 	/** Number of columns */
-	UPROPERTY(EditAnywhere, Category = "Calibration")
+	UPROPERTY(EditAnywhere, BlueprintReadWrite, Category = "Calibration")
 	int32 NumCornerCols = 2;
 
 	/** Length of the side of each square */
-	UPROPERTY(EditAnywhere, Category = "Calibration")
+	UPROPERTY(EditAnywhere, BlueprintReadWrite, Category = "Calibration")
 	float SquareSideLength = 3.2089f;
 
 	/** Thickness of checkerboard. Not used for calibration purposes. */
-	UPROPERTY(EditAnywhere, Category = "Calibration")
+	UPROPERTY(EditAnywhere, BlueprintReadWrite, Category = "Calibration")
 	float Thickness = 0.1f;
 
 	/** The static mesh that we are going to use for all the cubes */
