CameraCalibration: New SceneViewExtension for lens distortion that populates RDG textures for use with the experimental lens distortion post-processing pass.

diff --git a/Engine/Plugins/VirtualProduction/CameraCalibrationCore/Shaders/Private/BlendDisplacementMaps.usf b/Engine/Plugins/VirtualProduction/CameraCalibrationCore/Shaders/Private/BlendDisplacementMaps.usf
new file mode 100644
--- /dev/null
+++ b/Engine/Plugins/VirtualProduction/CameraCalibrationCore/Shaders/Private/BlendDisplacementMaps.usf
@@ -0,0 +1,127 @@
+// Copyright Epic Games, Inc. All Rights Reserved.
+
+/*=============================================================================
+	BlendDisplacementMaps.usf
+=============================================================================*/
+
+#include "/Engine/Public/Platform.ush"
+
+float2 ThreadIdToUV;
+float2 FxFyScale;
+
+// Stores the x and y values and the x and y tangents for each corner of the blending patch
+// Ordered (X0, Y0) -> (X1, Y0) -> (X1, Y1) -> (X0, Y1)
+float4 PatchCorners[4];
+
+float EvalFocus;
+float EvalZoom;
+
+Texture2D<float2> InputDistortionMap1;
+Texture2D<float2> InputDistortionMap2;
+Texture2D<float2> InputDistortionMap3;
+Texture2D<float2> InputDistortionMap4;
+
+SamplerState SourceTextureSampler;
+
+RWTexture2D<float2> OverscanDistortionMap;
+
+float2 BezierInterp(float2 P0, float2 P1, float2 P2, float2 P3, float Alpha)
+{
+	float2 P01 = lerp(P0, P1, Alpha);
+	float2 P12 = lerp(P1, P2, Alpha);
+	float2 P23 = lerp(P2, P3, Alpha);
+	float2 P012 = lerp(P01, P12, Alpha);
+	float2 P123 = lerp(P12, P23, Alpha);
+	float2 P0123 = lerp(P012, P123, Alpha);
+
+	return P0123;
+}
+
+float2 EvalAtTwoPoints(float EvalTime, float Time0, float Time1, float2 Value0, float2 Value1, float Tangent0, float Tangent1)
+{
+	if (Time0 == Time1)
+	{
+		return Value0;
+	}
+
+	float OneThird = 1.0f / 3.0f;
+
+	float Diff = Time1 - Time0;
+	float Alpha = (EvalTime - Time0) / Diff;
+
+	float2 DeltaValue = Value1 - Value0;
+	float2 TangentScale = DeltaValue / Diff;
+	float2 ScaledTangent0 = Tangent0 * TangentScale;
+	float2 ScaledTangent1 = Tangent1 * TangentScale;
+
+	float2 P0 = Value0;
+	float2 P3 = Value1;
+	float2 P1 = P0 + (ScaledTangent0 * Diff * OneThird);
+	float2 P2 = P3 - (ScaledTangent1 * Diff * OneThird);
+
+	return BezierInterp(P0, P1, P2, P3, Alpha);
+}
+
+float2 CoonsPatchBlend(float X, float Y, float2 Values[4])
+{
+	float Alpha = 0.0;
+	if (PatchCorners[0].x != PatchCorners[1].x)
+	{
+		Alpha = (X - PatchCorners[0].x) / (PatchCorners[1].x - PatchCorners[0].x);
+	}
+	
+	float Beta = 0.0;
+	if (PatchCorners[0].y != PatchCorners[3].y)
+	{
+		Beta = (Y - PatchCorners[0].y) / (PatchCorners[3].y - PatchCorners[0].y);
+	}
+	
+	float2 X0Curve = EvalAtTwoPoints(X, PatchCorners[0].x, PatchCorners[1].x, Values[0], Values[1], PatchCorners[0].z, PatchCorners[1].z);
+	float2 X1Curve = EvalAtTwoPoints(X, PatchCorners[3].x, PatchCorners[2].x, Values[3], Values[2], PatchCorners[3].z, PatchCorners[2].z);
+	float2 Y0Curve = EvalAtTwoPoints(Y, PatchCorners[0].y, PatchCorners[3].y, Values[0], Values[3], PatchCorners[0].w, PatchCorners[3].w);
+	float2 Y1Curve = EvalAtTwoPoints(Y, PatchCorners[1].y, PatchCorners[2].y, Values[1], Values[2], PatchCorners[1].w, PatchCorners[2].w);
+	
+	float2 Lx = lerp(X0Curve, X1Curve, Beta);
+	float2 Ly = lerp(Y0Curve, Y1Curve, Alpha);
+	float2 B = lerp(lerp(Values[0], Values[1], Alpha), lerp(Values[3], Values[2], Alpha), Beta);
+
+	return Lx + Ly - B;
+}
+
+[numthreads(8, 8, 1)]
+void MainCS(uint3 DispatchThreadId : SV_DispatchThreadID)
+{
+	float2 DestViewportUV = (DispatchThreadId.xy + 0.5) * ThreadIdToUV;
+
+	// FxFyScale alters the UV to sample the input maps to account for filmback cropping
+	// When the filmback is cropped, we only need to use a center crop of the distortion displacement map
+	DestViewportUV = ((DestViewportUV - 0.5) * FxFyScale) + 0.5;
+
+	float2 DistortionBlend;
+	float2 Distortion1 = InputDistortionMap1.SampleLevel(SourceTextureSampler, DestViewportUV, 0);
+
+#if BLEND_TYPE == 1 // Case - TwoFocusOneZoom - Requires two input distortion maps representing two focus points that share the same zoom value
+	float2 Distortion2 = InputDistortionMap2.SampleLevel(SourceTextureSampler, DestViewportUV, 0);
+	DistortionBlend = EvalAtTwoPoints(EvalFocus, PatchCorners[0].y, PatchCorners[1].y, Distortion1, Distortion2, PatchCorners[0].w, PatchCorners[1].w);
+	
+#elif BLEND_TYPE == 2 // Case - OneFocusTwoZoom - Requires two input distortion maps representing two zoom points that share the same focus value
+	float2 Distortion2 = InputDistortionMap2.SampleLevel(SourceTextureSampler, DestViewportUV, 0);
+	DistortionBlend = EvalAtTwoPoints(EvalZoom, PatchCorners[0].x, PatchCorners[1].x, Distortion1, Distortion2, PatchCorners[0].z, PatchCorners[1].z);
+	
+#elif BLEND_TYPE == 3 // Case - TwoFocusTwoZoom - Requires four input distortion maps representing four distinct focus/zoom pairs
+	float2 Distortions[4] =
+	{
+		Distortion1,
+		InputDistortionMap2.SampleLevel(SourceTextureSampler, DestViewportUV, 0),
+		InputDistortionMap3.SampleLevel(SourceTextureSampler, DestViewportUV, 0),
+		InputDistortionMap4.SampleLevel(SourceTextureSampler, DestViewportUV, 0)
+	};
+
+	DistortionBlend = CoonsPatchBlend(EvalZoom, EvalFocus, Distortions);
+
+#else // Case - OneFocusOneZoom - Requires one input distortion map which is passed through directly (after being scaled for filmback cropping)
+	DistortionBlend = Distortion1;
+#endif
+	
+	OverscanDistortionMap[DispatchThreadId.xy] = DistortionBlend;		
+}
diff --git a/Engine/Plugins/VirtualProduction/CameraCalibrationCore/Shaders/Private/CropDisplacementMap.usf b/Engine/Plugins/VirtualProduction/CameraCalibrationCore/Shaders/Private/CropDisplacementMap.usf
new file mode 100644
--- /dev/null
+++ b/Engine/Plugins/VirtualProduction/CameraCalibrationCore/Shaders/Private/CropDisplacementMap.usf
@@ -0,0 +1,20 @@
+// Copyright Epic Games, Inc. All Rights Reserved.
+
+/*=============================================================================
+	CropDisplacementMap.usf
+=============================================================================*/
+
+#include "/Engine/Public/Platform.ush"
+
+Texture2D<float2> InDistortionMapWithOverscan;
+RWTexture2D<float2> OutDistortionMap;
+
+// 2D offset from the top left of the overscan displacement map
+int2 OverscanOffset;
+
+[numthreads(8, 8, 1)]
+void MainCS(uint3 DispatchThreadId : SV_DispatchThreadID)
+{
+	uint2 OverscanIndex = DispatchThreadId.xy + OverscanOffset;
+	OutDistortionMap[DispatchThreadId.xy] = InDistortionMapWithOverscan[OverscanIndex];
+}
diff --git a/Engine/Plugins/VirtualProduction/CameraCalibrationCore/Shaders/Private/DrawDisplacementMaps.usf b/Engine/Plugins/VirtualProduction/CameraCalibrationCore/Shaders/Private/DrawDisplacementMaps.usf
new file mode 100644
--- /dev/null
+++ b/Engine/Plugins/VirtualProduction/CameraCalibrationCore/Shaders/Private/DrawDisplacementMaps.usf
@@ -0,0 +1,67 @@
+// Copyright Epic Games, Inc. All Rights Reserved.
+
+/*=============================================================================
+	DrawDisplacementMaps.usf
+=============================================================================*/
+
+#include "/Engine/Public/Platform.ush"
+
+float2 ThreadIdToUV;
+
+float2 FocalLength;
+float2 ImageCenter;
+float K1;
+float K2;
+float K3;
+float P1;
+float P2;
+
+float OverscanFactor;
+
+RWTexture2D<float2> OutDistortionMap;
+
+float2 SphericalDistort(float2 InputUV)
+{
+	float x_prime = (InputUV.x - 0.5) / FocalLength.x;
+	float y_prime = (InputUV.y - 0.5) / FocalLength.y;
+
+	float x_prime2 = x_prime * x_prime;
+	float y_prime2 = y_prime * y_prime;
+
+	float r2 = x_prime2 + y_prime2;
+	float r4 = r2 * r2;
+	float r6 = r2 * r2 * r2;
+
+	// Compute radial distortion
+	float dr = 1 + K1 * r2 + K2 * r4 + K3 * r6;
+
+	// Compute tangential distortion
+	float dtx = 2 * P1 * x_prime * y_prime + P2 * (r2 + 2 * x_prime2);
+	float dty = P1 * (r2 + 2 * y_prime2) + 2 * P2 * x_prime * y_prime;
+
+	// Compute total distortion
+	float x_double_prime = (x_prime * dr) + dtx;
+	float y_double_prime = (y_prime * dr) + dty;
+
+	float x_undistort = (FocalLength.x * x_double_prime) + ImageCenter.x;
+	float y_undistort = (FocalLength.y * y_double_prime) + ImageCenter.y;
+
+	float2 result;
+	result.x = x_undistort;
+	result.y = y_undistort;
+
+	return result;
+}
+
+[numthreads(8, 8, 1)]
+void MainCS(uint3 DispatchThreadId : SV_DispatchThreadID)
+{
+	// Add half-pixel offset to ThreadID to get accurate UV position
+	float2 UndistortedUV = (DispatchThreadId.xy + 0.5) * ThreadIdToUV;
+	
+	// Rescale the UV by an overscan factor (>= 1.0) to widen the range of UVs to be distorted
+	UndistortedUV = ((UndistortedUV - 0.5) * OverscanFactor) + 0.5;
+	
+	float2 DistortedUV = SphericalDistort(UndistortedUV);
+	OutDistortionMap[DispatchThreadId.xy] = DistortedUV - UndistortedUV;
+}
diff --git a/Engine/Plugins/VirtualProduction/CameraCalibrationCore/Shaders/Private/InvertDisplacementMap.usf b/Engine/Plugins/VirtualProduction/CameraCalibrationCore/Shaders/Private/InvertDisplacementMap.usf
new file mode 100644
--- /dev/null
+++ b/Engine/Plugins/VirtualProduction/CameraCalibrationCore/Shaders/Private/InvertDisplacementMap.usf
@@ -0,0 +1,53 @@
+// Copyright Epic Games, Inc. All Rights Reserved.
+
+/*=============================================================================
+	InvertDisplacementMap.usf
+=============================================================================*/
+
+#include "/Engine/Public/Platform.ush"
+
+uint2 GridDimensions;
+
+float2 PixelToUV;
+float2 PixelToOverscanUV;
+
+float OverscanFactor;
+
+Texture2D<float2> DistortionMap;
+SamplerState DistortionMapSampler;
+
+void MainVS(
+	in uint VertexId : SV_VertexID,
+	out noperspective float4 OutTexCoord : TEXCOORD0,
+	out float4 OutPosition : SV_POSITION)
+{
+	// Compute UV from vertex ID and rectangle layout
+	float2 ViewportUV = float2(VertexId % (GridDimensions.x + 1), VertexId / (GridDimensions.x + 1)) / float2(GridDimensions);
+
+	// Add half-pixel offset to get accurate UV position
+	ViewportUV += PixelToOverscanUV * 0.5;
+	
+	// Output the undistorted viewport UV, rescaled to account for the overscan of the distortion map
+	OutTexCoord.xy = ((ViewportUV.xy - 0.5) * OverscanFactor) + 0.5;
+	
+	float2 Displacement = DistortionMap.SampleLevel(DistortionMapSampler, ViewportUV, 0);
+	OutTexCoord.zw = Displacement; // Unused but useful to have this data in a graphics debugger
+
+	// Distort the location of the vertex, then rescale to [-1, 1] for clip space, and invert y
+	OutPosition.xy = OutTexCoord.xy + Displacement;
+	OutPosition.xy = ((OutPosition.xy * 2.0) - 1.0) * int2(1, -1);
+
+	OutPosition.zw = float2(0.0, 1.0);
+}
+
+
+void MainPS(
+	in noperspective float4 InTexCoord : TEXCOORD0,
+	in float4 SvPosition : SV_POSITION,
+	out float4 OutColor : SV_Target0)
+{
+	float2 DistortedUV = SvPosition.xy * PixelToUV;
+	float2 UndistortedDisplacement = InTexCoord.xy - DistortedUV;
+
+	OutColor = float4(UndistortedDisplacement, 0, 1);
+}
diff --git a/Engine/Plugins/VirtualProduction/CameraCalibrationCore/Source/CameraCalibrationCore/CameraCalibrationCore.Build.cs b/Engine/Plugins/VirtualProduction/CameraCalibrationCore/Source/CameraCalibrationCore/CameraCalibrationCore.Build.cs
--- a/Engine/Plugins/VirtualProduction/CameraCalibrationCore/Source/CameraCalibrationCore/CameraCalibrationCore.Build.cs
+++ b/Engine/Plugins/VirtualProduction/CameraCalibrationCore/Source/CameraCalibrationCore/CameraCalibrationCore.Build.cs
@@ -27,6 +27,7 @@ namespace UnrealBuildTool.Rules
                 {
 					"LiveLinkInterface",
                     "Projects",
+					"Renderer",
 				}
 			);
 
diff --git a/Engine/Plugins/VirtualProduction/CameraCalibrationCore/Source/CameraCalibrationCore/Private/CameraCalibrationSubsystem.cpp b/Engine/Plugins/VirtualProduction/CameraCalibrationCore/Source/CameraCalibrationCore/Private/CameraCalibrationSubsystem.cpp
--- a/Engine/Plugins/VirtualProduction/CameraCalibrationCore/Source/CameraCalibrationCore/Private/CameraCalibrationSubsystem.cpp
+++ b/Engine/Plugins/VirtualProduction/CameraCalibrationCore/Source/CameraCalibrationCore/Private/CameraCalibrationSubsystem.cpp
@@ -176,10 +176,22 @@ TArray<FName> UCameraCalibrationSubsystem::GetCameraCalibrationSteps() const
 	return OutKeys;
 }
 
+void UCameraCalibrationSubsystem::SetLensDistortionSVEState(ACameraActor* CameraActor, FDisplacementMapBlendingParams DistortionState)
+{
+	SceneViewExtension->UpdateDistortionState_AnyThread(CameraActor, DistortionState);
+}
+
+void UCameraCalibrationSubsystem::ClearLensDistortionSVEState(ACameraActor* CameraActor)
+{
+	SceneViewExtension->ClearDistortionState_AnyThread(CameraActor);
+}
+
 void UCameraCalibrationSubsystem::Initialize(FSubsystemCollectionBase& Collection)
 {
 	Super::Initialize(Collection);
 
+	SceneViewExtension = FSceneViewExtensions::NewExtension<FLensDistortionSceneViewExtension>();
+
 	PostEngineInitHandle = FCoreDelegates::OnPostEngineInit.AddLambda([&]()
 	{
 		// Find Image Center Algos
@@ -227,6 +239,9 @@ void UCameraCalibrationSubsystem::Initialize(FSubsystemCollectionBase& Collectio
 
 void UCameraCalibrationSubsystem::Deinitialize()
 {
+	SceneViewExtension.Reset();
+	SceneViewExtension = nullptr;
+	
 	LensModelMap.Empty(0);
 	CameraImageCenterAlgosMap.Empty(0);
 	CameraNodalOffsetAlgosMap.Empty(0);
diff --git a/Engine/Plugins/VirtualProduction/CameraCalibrationCore/Source/CameraCalibrationCore/Private/LensDistortionSceneViewExtension.cpp b/Engine/Plugins/VirtualProduction/CameraCalibrationCore/Source/CameraCalibrationCore/Private/LensDistortionSceneViewExtension.cpp
new file mode 100644
--- /dev/null
+++ b/Engine/Plugins/VirtualProduction/CameraCalibrationCore/Source/CameraCalibrationCore/Private/LensDistortionSceneViewExtension.cpp
@@ -0,0 +1,422 @@
+// Copyright Epic Games, Inc. All Rights Reserved.
+
+#include "LensDistortionSceneViewExtension.h"
+
+#include "Camera/CameraActor.h"
+#include "Containers/DynamicRHIResourceArray.h"
+#include "GlobalShader.h"
+#include "PostProcess/LensDistortion.h"
+#include "ScreenPass.h"
+#include "ShaderParameterStruct.h"
+#include "SystemTextures.h"
+
+TAutoConsoleVariable<float> CVarLensDistortionDisplacementOverscan(
+	TEXT("r.LensDistortion.DisplacementMapOverscan"),
+	1.25f,
+	TEXT("A factor to scale the distortion displacement map to ensure that the undistortion map is properly invertible.\n")
+	TEXT("Value is clamped between 1.0 (no overscan) and 2.0.\n"),
+	ECVF_RenderThreadSafe);
+
+TAutoConsoleVariable<int32> CVarLensDistortionInvertGridDensity(
+	TEXT("r.LensDistortion.InvertGridDensity"),
+	64,
+	TEXT("The number of squares drawn by the shader that inverts the distortion displacement map\n")
+	TEXT("Value is clamped between 64 and 255.\n"),
+	ECVF_RenderThreadSafe);
+
+FLensDistortionSceneViewExtension::FLensDistortionSceneViewExtension(const FAutoRegister& AutoRegister)
+	: FSceneViewExtensionBase(AutoRegister)
+{
+}
+
+void FLensDistortionSceneViewExtension::UpdateDistortionState_AnyThread(ACameraActor* CameraActor, FDisplacementMapBlendingParams DistortionState)
+{
+	FScopeLock ScopeLock(&DistortionStateMapCriticalSection);
+	DistortionStateMap.Add(CameraActor, DistortionState);
+}
+
+void FLensDistortionSceneViewExtension::ClearDistortionState_AnyThread(ACameraActor* CameraActor)
+{
+	FScopeLock ScopeLock(&DistortionStateMapCriticalSection);
+	DistortionStateMap.Remove(CameraActor);
+}
+
+class FDrawDistortionDisplacementMapCS : public FGlobalShader
+{
+	DECLARE_GLOBAL_SHADER(FDrawDistortionDisplacementMapCS);
+	SHADER_USE_PARAMETER_STRUCT(FDrawDistortionDisplacementMapCS, FGlobalShader);
+
+	BEGIN_SHADER_PARAMETER_STRUCT(FParameters, )
+		SHADER_PARAMETER(FVector2f, ThreadIdToUV)
+		SHADER_PARAMETER(FVector2f, FocalLength)
+		SHADER_PARAMETER(FVector2f, ImageCenter)
+		SHADER_PARAMETER(float, K1)
+		SHADER_PARAMETER(float, K2)
+		SHADER_PARAMETER(float, K3)
+		SHADER_PARAMETER(float, P1)
+		SHADER_PARAMETER(float, P2)
+		SHADER_PARAMETER(float, OverscanFactor)
+		SHADER_PARAMETER_RDG_TEXTURE_UAV(RWTexture2D, OutDistortionMap)
+	END_SHADER_PARAMETER_STRUCT()
+};
+
+IMPLEMENT_GLOBAL_SHADER(FDrawDistortionDisplacementMapCS, "/Plugin/CameraCalibrationCore/Private/DrawDisplacementMaps.usf", "MainCS", SF_Compute);
+
+
+class FBlendDistortionDisplacementMapCS : public FGlobalShader
+{
+	DECLARE_GLOBAL_SHADER(FBlendDistortionDisplacementMapCS);
+	SHADER_USE_PARAMETER_STRUCT(FBlendDistortionDisplacementMapCS, FGlobalShader);
+
+	class FBlendType : SHADER_PERMUTATION_INT("BLEND_TYPE", 4);
+	using FPermutationDomain = TShaderPermutationDomain<FBlendType>;
+
+	BEGIN_SHADER_PARAMETER_STRUCT(FParameters, )
+		SHADER_PARAMETER(FVector2f, ThreadIdToUV)
+		SHADER_PARAMETER(FVector2f, FxFyScale)
+		SHADER_PARAMETER_ARRAY(FVector4f, PatchCorners, [4])
+		SHADER_PARAMETER(float, EvalFocus)
+		SHADER_PARAMETER(float, EvalZoom)
+		SHADER_PARAMETER_RDG_TEXTURE_SRV(Texture2D, InputDistortionMap1)
+		SHADER_PARAMETER_RDG_TEXTURE_SRV(Texture2D, InputDistortionMap2)
+		SHADER_PARAMETER_RDG_TEXTURE_SRV(Texture2D, InputDistortionMap3)
+		SHADER_PARAMETER_RDG_TEXTURE_SRV(Texture2D, InputDistortionMap4)
+		SHADER_PARAMETER_SAMPLER(SamplerState, SourceTextureSampler)
+		SHADER_PARAMETER_RDG_TEXTURE_UAV(RWTexture2D, OverscanDistortionMap)
+	END_SHADER_PARAMETER_STRUCT()
+
+	// Called by the engine to determine which permutations to compile for this shader
+	static bool ShouldCompilePermutation(const FGlobalShaderPermutationParameters& Parameters)
+	{
+		return IsFeatureLevelSupported(Parameters.Platform, ERHIFeatureLevel::SM5);
+	}
+};
+
+IMPLEMENT_GLOBAL_SHADER(FBlendDistortionDisplacementMapCS, "/Plugin/CameraCalibrationCore/Private/BlendDisplacementMaps.usf", "MainCS", SF_Compute);
+
+class FCropDistortionDisplacementMapCS : public FGlobalShader
+{
+	DECLARE_GLOBAL_SHADER(FCropDistortionDisplacementMapCS);
+	SHADER_USE_PARAMETER_STRUCT(FCropDistortionDisplacementMapCS, FGlobalShader);
+
+	BEGIN_SHADER_PARAMETER_STRUCT(FParameters, )
+		SHADER_PARAMETER_RDG_TEXTURE_SRV(Texture2D, InDistortionMapWithOverscan)
+		SHADER_PARAMETER_RDG_TEXTURE_UAV(RWTexture2D, OutDistortionMap)
+		SHADER_PARAMETER(FIntPoint, OverscanOffset)
+	END_SHADER_PARAMETER_STRUCT()
+};
+
+IMPLEMENT_GLOBAL_SHADER(FCropDistortionDisplacementMapCS, "/Plugin/CameraCalibrationCore/Private/CropDisplacementMap.usf", "MainCS", SF_Compute);
+
+BEGIN_SHADER_PARAMETER_STRUCT(FInvertDisplacementParameters, )
+	SHADER_PARAMETER(FIntPoint, GridDimensions)
+	SHADER_PARAMETER(FVector2f, PixelToUV)
+	SHADER_PARAMETER(FVector2f, PixelToOverscanUV)
+	SHADER_PARAMETER(float, OverscanFactor)
+	SHADER_PARAMETER_RDG_TEXTURE_SRV(Texture2D, DistortionMap)
+	SHADER_PARAMETER_SAMPLER(SamplerState, DistortionMapSampler)
+	RENDER_TARGET_BINDING_SLOTS()
+END_SHADER_PARAMETER_STRUCT()
+
+class FInvertDisplacementVS : public FGlobalShader
+{
+public:
+	DECLARE_GLOBAL_SHADER(FInvertDisplacementVS);
+	SHADER_USE_PARAMETER_STRUCT(FInvertDisplacementVS, FGlobalShader);
+	using FParameters = FInvertDisplacementParameters;
+};
+
+class FInvertDisplacementPS : public FGlobalShader
+{
+public:
+	DECLARE_GLOBAL_SHADER(FInvertDisplacementPS);
+	SHADER_USE_PARAMETER_STRUCT(FInvertDisplacementPS, FGlobalShader);
+	using FParameters = FInvertDisplacementParameters;
+};
+
+IMPLEMENT_GLOBAL_SHADER(FInvertDisplacementVS, "/Plugin/CameraCalibrationCore/Private/InvertDisplacementMap.usf", "MainVS", SF_Vertex);
+IMPLEMENT_GLOBAL_SHADER(FInvertDisplacementPS, "/Plugin/CameraCalibrationCore/Private/InvertDisplacementMap.usf", "MainPS", SF_Pixel);
+
+void FLensDistortionSceneViewExtension::DrawDisplacementMap_RenderThread(FRDGBuilder& GraphBuilder, const FLensDistortionState& CurrentState, float OverscanFactor, FRDGTextureRef& OutDistortionMapWithOverscan)
+{
+	if (CurrentState.DistortionInfo.Parameters.IsEmpty())
+	{
+		OutDistortionMapWithOverscan = GSystemTextures.GetBlackDummy(GraphBuilder);
+		return;
+	}
+
+	FDrawDistortionDisplacementMapCS::FParameters* PassParameters = GraphBuilder.AllocParameters<FDrawDistortionDisplacementMapCS::FParameters>();
+
+	PassParameters->OutDistortionMap = GraphBuilder.CreateUAV(OutDistortionMapWithOverscan);
+
+	FIntPoint DistortionMapResolution = OutDistortionMapWithOverscan->Desc.Extent;
+	PassParameters->ThreadIdToUV = FVector2f(1.0f) / FVector2f(DistortionMapResolution);
+
+	PassParameters->ImageCenter = FVector2f(CurrentState.ImageCenter.PrincipalPoint);
+	PassParameters->FocalLength = FVector2f(CurrentState.FocalLengthInfo.FxFy);
+
+	PassParameters->K1 = CurrentState.DistortionInfo.Parameters[0];
+	PassParameters->K2 = CurrentState.DistortionInfo.Parameters[1];
+	PassParameters->K3 = CurrentState.DistortionInfo.Parameters[2];
+	PassParameters->P1 = CurrentState.DistortionInfo.Parameters[3];
+	PassParameters->P2 = CurrentState.DistortionInfo.Parameters[4];
+
+	PassParameters->OverscanFactor = OverscanFactor;
+
+	TShaderMapRef<FDrawDistortionDisplacementMapCS> ComputeShader(GetGlobalShaderMap(GMaxRHIFeatureLevel));
+	FComputeShaderUtils::AddPass(
+		GraphBuilder,
+		RDG_EVENT_NAME("DrawDistortionDisplacementMap"),
+		ComputeShader,
+		PassParameters,
+		FIntVector(FMath::DivideAndRoundUp(DistortionMapResolution.X, 8), FMath::DivideAndRoundUp(DistortionMapResolution.Y, 8), 1));
+}
+
+void FLensDistortionSceneViewExtension::CropDisplacementMap_RenderThread(FRDGBuilder& GraphBuilder, const FRDGTextureRef& InDistortionMapWithOverscan, FRDGTextureRef& OutDistortionMap)
+{
+	FCropDistortionDisplacementMapCS::FParameters* PassParameters = GraphBuilder.AllocParameters<FCropDistortionDisplacementMapCS::FParameters>();
+
+	PassParameters->InDistortionMapWithOverscan = GraphBuilder.CreateSRV(InDistortionMapWithOverscan);
+	PassParameters->OutDistortionMap = GraphBuilder.CreateUAV(OutDistortionMap);
+
+	FIntPoint LUTResolution = OutDistortionMap->Desc.Extent;
+	PassParameters->OverscanOffset = (InDistortionMapWithOverscan->Desc.Extent - OutDistortionMap->Desc.Extent) / 2;
+
+	TShaderMapRef<FCropDistortionDisplacementMapCS> ComputeShader(GetGlobalShaderMap(GMaxRHIFeatureLevel));
+	FComputeShaderUtils::AddPass(
+		GraphBuilder,
+		RDG_EVENT_NAME("CropDistortionDisplacementMap"),
+		ComputeShader,
+		PassParameters,
+		FIntVector(FMath::DivideAndRoundUp(LUTResolution.X, 8), FMath::DivideAndRoundUp(LUTResolution.Y, 8), 1));
+}
+
+void FLensDistortionSceneViewExtension::BlendDisplacementMaps_RenderThread(FRDGBuilder& GraphBuilder, const FDisplacementMapBlendingParams& BlendState, float OverscanFactor, FRDGTextureRef& OutDistortionMapWithOverscan)
+{
+	FBlendDistortionDisplacementMapCS::FParameters* PassParameters = GraphBuilder.AllocParameters<FBlendDistortionDisplacementMapCS::FParameters>();
+
+	// Draw the first distortion map, which should always be valid
+	{
+		FRDGTextureRef Distortion1 = GraphBuilder.CreateTexture(OutDistortionMapWithOverscan->Desc, TEXT("DistortingDisplacement1"));
+		DrawDisplacementMap_RenderThread(GraphBuilder, BlendState.States[0], OverscanFactor, Distortion1);
+		PassParameters->InputDistortionMap1 = GraphBuilder.CreateSRV(Distortion1);
+	}
+
+	// Draw the second distortion map if any blend is needed
+	if (BlendState.BlendType != EDisplacementMapBlendType::OneFocusOneZoom)
+	{
+		FRDGTextureRef Distortion2 = GraphBuilder.CreateTexture(OutDistortionMapWithOverscan->Desc, TEXT("DistortingDisplacement2"));
+		DrawDisplacementMap_RenderThread(GraphBuilder, BlendState.States[1], OverscanFactor, Distortion2);
+		PassParameters->InputDistortionMap2 = GraphBuilder.CreateSRV(Distortion2);
+	}
+
+	// Draw the 3rd and 4th distortion maps if a 4-way blend is needed
+	if (BlendState.BlendType == EDisplacementMapBlendType::TwoFocusTwoZoom)
+	{
+		FRDGTextureRef Distortion3 = GraphBuilder.CreateTexture(OutDistortionMapWithOverscan->Desc, TEXT("DistortingDisplacement3"));
+		FRDGTextureRef Distortion4 = GraphBuilder.CreateTexture(OutDistortionMapWithOverscan->Desc, TEXT("DistortingDisplacement4"));
+
+		DrawDisplacementMap_RenderThread(GraphBuilder, BlendState.States[2], OverscanFactor, Distortion3);
+		DrawDisplacementMap_RenderThread(GraphBuilder, BlendState.States[3], OverscanFactor, Distortion4);
+
+		PassParameters->InputDistortionMap3 = GraphBuilder.CreateSRV(Distortion3);
+		PassParameters->InputDistortionMap4 = GraphBuilder.CreateSRV(Distortion4);
+	}
+
+	PassParameters->OverscanDistortionMap = GraphBuilder.CreateUAV(OutDistortionMapWithOverscan);
+	PassParameters->SourceTextureSampler = TStaticSamplerState<SF_Bilinear, AM_Clamp, AM_Clamp, AM_Clamp>::GetRHI();
+
+	const FIntPoint DistortionMapResolution = OutDistortionMapWithOverscan->Desc.Extent;
+	PassParameters->ThreadIdToUV = FVector2f(1.0f / DistortionMapResolution.X, 1.0f / DistortionMapResolution.Y);
+
+	PassParameters->FxFyScale = FVector2f(BlendState.FxFyScale);
+
+	// Set permutation and blending params based on blend type
+	PassParameters->EvalFocus = BlendState.EvalFocus;
+	PassParameters->EvalZoom = BlendState.EvalZoom;
+
+	FBlendDistortionDisplacementMapCS::FPermutationDomain PermutationVector;
+	if (BlendState.BlendType == EDisplacementMapBlendType::OneFocusOneZoom)
+	{
+		PermutationVector.Set<FBlendDistortionDisplacementMapCS::FBlendType>(0);
+	}
+	else if (BlendState.BlendType == EDisplacementMapBlendType::TwoFocusOneZoom)
+	{
+		PermutationVector.Set<FBlendDistortionDisplacementMapCS::FBlendType>(1);
+		PassParameters->PatchCorners[0] = BlendState.PatchCorners[0].ToVector();
+		PassParameters->PatchCorners[1] = BlendState.PatchCorners[1].ToVector();
+		PassParameters->PatchCorners[2] = FVector4f::Zero();
+		PassParameters->PatchCorners[3] = FVector4f::Zero();
+	}
+	else if (BlendState.BlendType == EDisplacementMapBlendType::OneFocusTwoZoom)
+	{
+		PermutationVector.Set<FBlendDistortionDisplacementMapCS::FBlendType>(2);
+		PassParameters->PatchCorners[0] = BlendState.PatchCorners[0].ToVector();
+		PassParameters->PatchCorners[1] = BlendState.PatchCorners[1].ToVector();
+		PassParameters->PatchCorners[2] = FVector4f::Zero();
+		PassParameters->PatchCorners[3] = FVector4f::Zero();
+	}
+	else if (BlendState.BlendType == EDisplacementMapBlendType::TwoFocusTwoZoom)
+	{
+		PermutationVector.Set<FBlendDistortionDisplacementMapCS::FBlendType>(3);
+		PassParameters->PatchCorners[0] = BlendState.PatchCorners[0].ToVector();
+		PassParameters->PatchCorners[1] = BlendState.PatchCorners[1].ToVector();
+		PassParameters->PatchCorners[2] = BlendState.PatchCorners[2].ToVector();
+		PassParameters->PatchCorners[3] = BlendState.PatchCorners[3].ToVector();
+	}
+
+	TShaderMapRef<FBlendDistortionDisplacementMapCS> ComputeShader(GetGlobalShaderMap(GMaxRHIFeatureLevel), PermutationVector);
+	FComputeShaderUtils::AddPass(
+		GraphBuilder,
+		RDG_EVENT_NAME("BlendDistortionDisplacementMap"),
+		ComputeShader,
+		PassParameters,
+		FIntVector(FMath::DivideAndRoundUp(DistortionMapResolution.X, 8), FMath::DivideAndRoundUp(DistortionMapResolution.Y, 8), 1));
+}
+
+void FLensDistortionSceneViewExtension::InvertDistortionMap_RenderThread(FRDGBuilder& GraphBuilder, const FRDGTextureRef& InDistortionMap, FRDGTextureRef& OutUndistortionMap)
+{
+	FInvertDisplacementParameters* PassParameters = GraphBuilder.AllocParameters<FInvertDisplacementParameters>();
+
+	FScreenPassRenderTarget Output;
+	Output.Texture = OutUndistortionMap;
+	Output.ViewRect = FIntRect(FIntPoint(0, 0), OutUndistortionMap->Desc.Extent);
+	Output.LoadAction = ERenderTargetLoadAction::EClear;
+	Output.UpdateVisualizeTextureExtent();
+
+	PassParameters->RenderTargets[0] = Output.GetRenderTargetBinding();
+
+	PassParameters->DistortionMap = GraphBuilder.CreateSRV(InDistortionMap);
+	PassParameters->DistortionMapSampler = TStaticSamplerState<SF_Bilinear, AM_Clamp, AM_Clamp, AM_Clamp>::GetRHI();
+
+	PassParameters->OverscanFactor = float(InDistortionMap->Desc.Extent.X) / float(OutUndistortionMap->Desc.Extent.X);
+	PassParameters->PixelToUV = FVector2f(1.0f) / FVector2f(OutUndistortionMap->Desc.Extent);
+	PassParameters->PixelToOverscanUV = FVector2f(1.0f) / FVector2f(InDistortionMap->Desc.Extent);
+
+	GraphBuilder.AddPass(
+		RDG_EVENT_NAME("InvertDistortionDisplacementMap"),
+		PassParameters,
+		ERDGPassFlags::Raster,
+		[PassParameters, Output](FRHICommandList& RHICmdList)
+		{
+			const int32 NumSquares = FMath::Clamp(CVarLensDistortionInvertGridDensity.GetValueOnRenderThread(), 64, 255);
+			const FIntPoint GridDimensions = FIntPoint(NumSquares);
+			PassParameters->GridDimensions = GridDimensions;
+
+			RHICmdList.SetViewport(Output.ViewRect.Min.X, Output.ViewRect.Min.Y, 0.0f, Output.ViewRect.Max.X, Output.ViewRect.Max.Y, 1.0f);
+
+			TShaderMapRef<FInvertDisplacementVS> VertexShader(GetGlobalShaderMap(GMaxRHIFeatureLevel));
+			TShaderMapRef<FInvertDisplacementPS> PixelShader(GetGlobalShaderMap(GMaxRHIFeatureLevel));
+
+			SetScreenPassPipelineState(RHICmdList, FScreenPassPipelineState(VertexShader, PixelShader));
+			SetShaderParameters(RHICmdList, VertexShader, VertexShader.GetVertexShader(), *PassParameters);
+			SetShaderParameters(RHICmdList, PixelShader, PixelShader.GetPixelShader(), *PassParameters);
+
+			FRHIBatchedShaderParameters& BatchedParameters = RHICmdList.GetScratchShaderParameters();
+			RHICmdList.SetBatchedShaderParameters(VertexShader.GetVertexShader(), BatchedParameters);
+
+			// No vertex buffer is needed because we compute it in the VS
+			RHICmdList.SetStreamSource(0, NULL, 0);
+
+			// The following code for setting up this index buffer based on FTesselatedScreenRectangleIndexBuffer::InitRHI()
+			TResourceArray<uint16> IndexBuffer;
+
+			const uint32 Width = GridDimensions.X;
+			const uint32 Height = GridDimensions.Y;
+			const uint32 NumTriangles = Width * Height * 2;
+			const uint32 NumIndices = NumTriangles * 3;
+
+			IndexBuffer.AddUninitialized(NumIndices);
+
+			uint16* IndexBufferData = (uint16*)IndexBuffer.GetData();
+
+			for (uint32 IndexY = 0; IndexY < Height; ++IndexY)
+			{
+				for (uint32 IndexX = 0; IndexX < Width; ++IndexX)
+				{
+					// left top to bottom right in reading order
+					uint16 Index00 = IndexX + IndexY * (Width + 1);
+					uint16 Index10 = Index00 + 1;
+					uint16 Index01 = Index00 + (Width + 1);
+					uint16 Index11 = Index01 + 1;
+
+					// triangle A
+					*IndexBufferData++ = Index00;
+					*IndexBufferData++ = Index01;
+					*IndexBufferData++ = Index10;
+
+					// triangle B
+					*IndexBufferData++ = Index11;
+					*IndexBufferData++ = Index10;
+					*IndexBufferData++ = Index01;
+				}
+			}
+
+			// Create index buffer. Fill buffer with initial data upon creation
+			FRHIResourceCreateInfo CreateInfo(TEXT("InvertDistortionMapIndexBuffer"), &IndexBuffer);
+			FBufferRHIRef IndexBufferRHI = RHICmdList.CreateIndexBuffer(sizeof(uint16), IndexBuffer.GetResourceDataSize(), BUF_Static, CreateInfo);
+
+			RHICmdList.DrawIndexedPrimitive(IndexBufferRHI, 0, 0, NumIndices, 0, NumTriangles, 1);
+		});
+}
+
+void FLensDistortionSceneViewExtension::PreRenderView_RenderThread(FRDGBuilder& GraphBuilder, FSceneView& InView)
+{
+	if (const ACameraActor* CameraActor = Cast<const ACameraActor>(InView.ViewActor))
+	{
+		FDisplacementMapBlendingParams BlendState;
+		{
+			FScopeLock ScopeLock(&DistortionStateMapCriticalSection);
+			if (FDisplacementMapBlendingParams* BlendStatePtr = DistortionStateMap.Find(CameraActor))
+			{
+				BlendState = *BlendStatePtr;
+			}
+			else
+			{
+				return;
+			}
+		}
+
+		// Create the distortion map and undistortion map textures for the FLensDistortionLUT for this frame
+		const FIntPoint DisplacementMapResolution = FIntPoint(256, 256);
+
+		FRDGTextureDesc DistortionMapDesc = FRDGTextureDesc::Create2D(
+			DisplacementMapResolution,
+			PF_G32R32F,
+			FClearValueBinding::None,
+			TexCreate_ShaderResource | TexCreate_UAV);
+
+		FRDGTextureDesc UndistortionMapDesc = FRDGTextureDesc::Create2D(
+			DisplacementMapResolution,
+			PF_G32R32F,
+			FClearValueBinding::Black,
+			TexCreate_ShaderResource | TexCreate_RenderTargetable);
+
+		FLensDistortionLUT ViewDistortionLUT;
+
+		ViewDistortionLUT.DistortingDisplacementTexture = GraphBuilder.CreateTexture(DistortionMapDesc, TEXT("DistortionDisplacementMap"));
+		ViewDistortionLUT.UndistortingDisplacementTexture = GraphBuilder.CreateTexture(UndistortionMapDesc, TEXT("UndistortionDisplacementMap"));
+
+		// In order to guarantee that we can generate a complete undistortion map, the distortion map we invert needs to have some overscan 
+		float OverscanFactor = FMath::Clamp(CVarLensDistortionDisplacementOverscan.GetValueOnRenderThread(), 1.0f, 2.0f);
+
+		// Adjust the overscan resolution to be square, with each side being a multiple of 8
+		const FIntPoint OverscanResolution = FIntPoint(FMath::CeilToInt(OverscanFactor * 32) * 8);
+		OverscanFactor = float(OverscanResolution.X) / float(DisplacementMapResolution.X);
+
+		// Create the texture for the overscanned distortion map
+		FRDGTextureDesc OverscanDesc = FRDGTextureDesc::Create2D(
+			OverscanResolution,
+			PF_G32R32F,
+			FClearValueBinding::None,
+			TexCreate_ShaderResource | TexCreate_UAV);
+
+		FRDGTextureRef DistortionMapWithOverscan = GraphBuilder.CreateTexture(OverscanDesc, TEXT("DistortionMapWithOverscan"));
+
+		BlendDisplacementMaps_RenderThread(GraphBuilder, BlendState, OverscanFactor, DistortionMapWithOverscan);
+		InvertDistortionMap_RenderThread(GraphBuilder, DistortionMapWithOverscan, ViewDistortionLUT.UndistortingDisplacementTexture);
+		CropDisplacementMap_RenderThread(GraphBuilder, DistortionMapWithOverscan, ViewDistortionLUT.DistortingDisplacementTexture);
+
+		LensDistortion::SetLUTUnsafe(InView, ViewDistortionLUT);
+	}
+}
diff --git a/Engine/Plugins/VirtualProduction/CameraCalibrationCore/Source/CameraCalibrationCore/Private/LensFile.cpp b/Engine/Plugins/VirtualProduction/CameraCalibrationCore/Source/CameraCalibrationCore/Private/LensFile.cpp
--- a/Engine/Plugins/VirtualProduction/CameraCalibrationCore/Source/CameraCalibrationCore/Private/LensFile.cpp
+++ b/Engine/Plugins/VirtualProduction/CameraCalibrationCore/Source/CameraCalibrationCore/Private/LensFile.cpp
@@ -405,6 +405,36 @@ void ULensFile::SetupNoDistortionOutput(ULensDistortionModelHandlerBase* LensHan
 	LensHandler->SetOverscanFactor(1.0f);
 }
 
+void ULensFile::GetBlendState(float InFocus, float InZoom, FVector2D InFilmback, FDisplacementMapBlendingParams& OutBlendState)
+{
+	LensInterpolationUtils::FDistortionMapBlendParams<FDistortionTable> Params;
+	Params.bGenerateBlendingParams = true;
+
+ 	FImageCenterInfo InterpolatedImageCenter;
+ 	EvaluateImageCenterParameters(InFocus, InZoom, InterpolatedImageCenter);
+
+	Params.GetDistortionState = LensInterpolationUtils::FDistortionMapBlendParams<FDistortionTable>::FGetDistortionState::CreateLambda(
+		[this, &InterpolatedImageCenter]
+		(const FDistortionFocusPoint& FocusPoint, const FDistortionFocusCurve& FocusCurve, FLensDistortionState& OutState)
+		{
+			// In case the point doesn't exist, we need to fill the distortion parameter array with default values
+			LensInfo.LensModel->GetDefaultObject<ULensModel>()->GetDefaultParameterArray(OutState.DistortionInfo.Parameters);
+			FocusPoint.GetPoint(FocusCurve.Zoom, OutState.DistortionInfo);
+
+			LensDataTableUtils::GetPointValue<FFocalLengthFocusPoint>(FocusPoint.Focus, FocusCurve.Zoom, FocalLengthTable.FocusPoints, OutState.FocalLengthInfo);
+			OutState.ImageCenter = InterpolatedImageCenter;
+		}
+	);
+
+	LensInterpolationUtils::FDistortionMapBlendResults Results = LensInterpolationUtils::DistortionMapBlend(DistortionTable, InFocus, InZoom, Params);
+
+	if (Results.BlendingParams.IsSet())
+	{
+		OutBlendState = Results.BlendingParams.GetValue();
+		OutBlendState.FxFyScale = FVector2D(InFilmback.X / LensInfo.SensorDimensions.X, InFilmback.Y / LensInfo.SensorDimensions.Y);
+	}
+}
+
 bool ULensFile::EvaluateDistortionForParameters(float InFocus, float InZoom, FVector2D InFilmback, ULensDistortionModelHandlerBase* InLensHandler) const
 {
 	// Compute interpolated image center and focal length to pass to the lens handler
diff --git a/Engine/Plugins/VirtualProduction/CameraCalibrationCore/Source/CameraCalibrationCore/Private/LensFileRendering.h b/Engine/Plugins/VirtualProduction/CameraCalibrationCore/Source/CameraCalibrationCore/Private/LensFileRendering.h
deleted file mode 100644
diff --git a/Engine/Plugins/VirtualProduction/CameraCalibrationCore/Source/CameraCalibrationCore/Private/LensInterpolationUtils.h b/Engine/Plugins/VirtualProduction/CameraCalibrationCore/Source/CameraCalibrationCore/Private/LensInterpolationUtils.h
--- a/Engine/Plugins/VirtualProduction/CameraCalibrationCore/Source/CameraCalibrationCore/Private/LensInterpolationUtils.h
+++ b/Engine/Plugins/VirtualProduction/CameraCalibrationCore/Source/CameraCalibrationCore/Private/LensInterpolationUtils.h
@@ -3,10 +3,10 @@
 #pragma once
 
 #include "CoreTypes.h"
-#include "LensFileRendering.h"
 #include "Curves/RichCurve.h"
+#include "LensDistortionSceneViewExtension.h"
+#include "LensFileRendering.h"
 #include "Math/Vector2D.h"
-
 #include "Misc/EnumClassFlags.h"
 #include "Tables/DistortionParametersTable.h"
 #include "Tables/LensTableUtils.h"
@@ -243,6 +243,11 @@ namespace LensInterpolationUtils
 			UTextureRenderTarget2D* /* InUndistortedMap */,
 			UTextureRenderTarget2D* /* InDistortedMap */);
 
+		DECLARE_DELEGATE_ThreeParams(FGetDistortionState,
+			const FocusPointType& /* FocusPoint */,
+			const FocusCurveType& /* FocusCurve */,
+			FLensDistortionState& /* OutState */);
+
 		/**
 		 * Callback to use to retrieve the distortion parameters for a specified point being blended.
 		 * If unbound, no distortion parameters will be computed in the blending results
@@ -260,7 +265,13 @@ namespace LensInterpolationUtils
 		 * If unbound, no overscan blending will be computed in the blending results
 		 */
 		FProcessDisplacementMaps ProcessDisplacementMaps;
-		
+
+		/**
+		 * Callback to use to retrieve the complete distortion state for a specified point being blended.
+		 * If unbound, no distortion state will be added to the blending results
+		 */
+		FGetDistortionState GetDistortionState;
+
 		/** Indicates that shader blending parameters should be calculated */
 		bool bGenerateBlendingParams = false;
 		
@@ -278,7 +289,7 @@ namespace LensInterpolationUtils
 	/** Blended results from a distortion map blend. Optionals are set when the blend input parameters indicated those blended values should be computed */
 	struct FDistortionMapBlendResults
 	{
-		/** Indicates that a distortion map blend successfully occured */
+		/** Indicates that a distortion map blend successfully occurred */
 		bool bValid = false;
 
 		/** The shader blending parameters, if they were computed for the blend */
@@ -394,6 +405,7 @@ namespace LensInterpolationUtils
 			if (Results.BlendingParams.IsSet())
 			{
 				Results.BlendingParams->BlendType = EDisplacementMapBlendType::OneFocusOneZoom;
+				InParams.GetDistortionState.ExecuteIfBound(PrevFocusPoint, PrevFocusCurve, Results.BlendingParams->States[0]);
 			}
 
 			if (Results.BlendedDistortionParams.IsSet())
@@ -424,6 +436,9 @@ namespace LensInterpolationUtils
 				Results.BlendingParams->BlendType = EDisplacementMapBlendType::OneFocusTwoZoom;
 				Results.BlendingParams->PatchCorners[0] = FDisplacementMapBlendPatchCorner(PrevFocusCurve.Zoom, PrevFocusPoint.Focus, PrevTangent, 0.0);
 				Results.BlendingParams->PatchCorners[1] = FDisplacementMapBlendPatchCorner(NextFocusCurve.Zoom, PrevFocusPoint.Focus, NextTangent, 0.0);
+
+				InParams.GetDistortionState.ExecuteIfBound(PrevFocusPoint, PrevFocusCurve, Results.BlendingParams->States[0]);
+				InParams.GetDistortionState.ExecuteIfBound(PrevFocusPoint, NextFocusCurve, Results.BlendingParams->States[1]);
 			}
 
 			if (Results.BlendedDistortionParams.IsSet())
@@ -467,6 +482,9 @@ namespace LensInterpolationUtils
 				Results.BlendingParams->BlendType = EDisplacementMapBlendType::TwoFocusOneZoom;
 				Results.BlendingParams->PatchCorners[0] = FDisplacementMapBlendPatchCorner(PrevFocusCurve.Zoom, PrevFocusPoint.Focus, 0.0, PrevTangent);
 				Results.BlendingParams->PatchCorners[1] = FDisplacementMapBlendPatchCorner(PrevFocusCurve.Zoom, NextFocusPoint.Focus, 0.0, NextTangent);
+
+				InParams.GetDistortionState.ExecuteIfBound(PrevFocusPoint, PrevFocusCurve, Results.BlendingParams->States[0]);
+				InParams.GetDistortionState.ExecuteIfBound(NextFocusPoint, PrevFocusCurve, Results.BlendingParams->States[1]);
 			}
 			
 			if (Results.BlendedDistortionParams.IsSet())
@@ -517,6 +535,11 @@ namespace LensInterpolationUtils
 				Results.BlendingParams->PatchCorners[1] = FDisplacementMapBlendPatchCorner(NextFocusCurve.Zoom, PrevFocusPoint.Focus, XTangents[1], YTangents[1]);
 				Results.BlendingParams->PatchCorners[2] = FDisplacementMapBlendPatchCorner(NextFocusCurve.Zoom, NextFocusPoint.Focus, XTangents[2], YTangents[2]);
 				Results.BlendingParams->PatchCorners[3] = FDisplacementMapBlendPatchCorner(PrevFocusCurve.Zoom, NextFocusPoint.Focus, XTangents[3], YTangents[3]);
+
+				InParams.GetDistortionState.ExecuteIfBound(PrevFocusPoint, PrevFocusCurve, Results.BlendingParams->States[0]);
+				InParams.GetDistortionState.ExecuteIfBound(PrevFocusPoint, NextFocusCurve, Results.BlendingParams->States[1]);
+				InParams.GetDistortionState.ExecuteIfBound(NextFocusPoint, NextFocusCurve, Results.BlendingParams->States[2]);
+				InParams.GetDistortionState.ExecuteIfBound(NextFocusPoint, PrevFocusCurve, Results.BlendingParams->States[3]);
 			}
 			
 			if (Results.BlendedDistortionParams.IsSet())
diff --git a/Engine/Plugins/VirtualProduction/CameraCalibrationCore/Source/CameraCalibrationCore/Public/CameraCalibrationSubsystem.h b/Engine/Plugins/VirtualProduction/CameraCalibrationCore/Source/CameraCalibrationCore/Public/CameraCalibrationSubsystem.h
--- a/Engine/Plugins/VirtualProduction/CameraCalibrationCore/Source/CameraCalibrationCore/Public/CameraCalibrationSubsystem.h
+++ b/Engine/Plugins/VirtualProduction/CameraCalibrationCore/Source/CameraCalibrationCore/Public/CameraCalibrationSubsystem.h
@@ -12,6 +12,7 @@
 #include "CameraCalibrationStep.h"
 #include "Containers/ArrayView.h"
 #include "LensDistortionModelHandlerBase.h"
+#include "LensDistortionSceneViewExtension.h"
 #include "LensFile.h"
 #include "UObject/ObjectKey.h"
 
@@ -117,6 +118,12 @@ public:
 	/** Unregister an overlay material */
 	void UnregisterOverlayMaterial(const FName& MaterialName);
 
+	/** Adds the input distortion state and blending parameters to the Lens Distortion Scene View Extension for the input camera actor */
+	void SetLensDistortionSVEState(ACameraActor* CameraActor, FDisplacementMapBlendingParams DistortionState);
+
+	/** Removes all distortion state from the Lens Distortion Scene View Extension for the input camera actor */
+	void ClearLensDistortionSVEState(ACameraActor* CameraActor);
+
 private:
 	/** Default lens file to use when no override has been provided */
 	UPROPERTY(Transient)
@@ -144,6 +151,9 @@ private:
 	/** Map of actor components to the authoritative lens model that should be used with that component */
 	TMap<FObjectKey, TSubclassOf<ULensModel>> ComponentsWithAuthoritativeModels;
 
+	/** Lens Distortion Scene View Extension, used to render distortion and undistortion displacement maps */
+	TSharedPtr<FLensDistortionSceneViewExtension, ESPMode::ThreadSafe> SceneViewExtension;
+
 private:
 
 	FDelegateHandle PostEngineInitHandle;
diff --git a/Engine/Plugins/VirtualProduction/CameraCalibrationCore/Source/CameraCalibrationCore/Public/LensDistortionSceneViewExtension.h b/Engine/Plugins/VirtualProduction/CameraCalibrationCore/Source/CameraCalibrationCore/Public/LensDistortionSceneViewExtension.h
new file mode 100644
--- /dev/null
+++ b/Engine/Plugins/VirtualProduction/CameraCalibrationCore/Source/CameraCalibrationCore/Public/LensDistortionSceneViewExtension.h
@@ -0,0 +1,52 @@
+// Copyright Epic Games, Inc. All Rights Reserved.
+
+#pragma once
+
+#include "SceneViewExtension.h"
+
+#include "LensFileRendering.h"
+
+/**
+ * View extension drawing distortion/undistortion displacement maps
+ */
+class FLensDistortionSceneViewExtension : public FSceneViewExtensionBase
+{
+public:
+
+	FLensDistortionSceneViewExtension(const FAutoRegister& AutoRegister);
+
+	//~ Begin ISceneViewExtension interface	
+	virtual void SetupViewFamily(FSceneViewFamily& InViewFamily) override {};
+	virtual void SetupView(FSceneViewFamily& InViewFamily, FSceneView& InView) override {};
+	virtual void BeginRenderViewFamily(FSceneViewFamily& InViewFamily) override {};
+
+	virtual void PreRenderView_RenderThread(FRDGBuilder& GraphBuilder, FSceneView& InView) override;
+	//~End ISceneVIewExtension interface
+
+public:
+	/** Update the distortion state and blending params for the input camera */
+	void UpdateDistortionState_AnyThread(ACameraActor* CameraActor, FDisplacementMapBlendingParams DistortionState);
+
+	/** Remove the distortion state and blending params for the input camera */
+	void ClearDistortionState_AnyThread(ACameraActor* CameraActor);
+
+private:
+	/** Use the input distortion state to draw a distortion displacement map */
+	void DrawDisplacementMap_RenderThread(FRDGBuilder& GraphBuilder, const FLensDistortionState& CurrentState, float OverscanFactor, FRDGTextureRef& OutDistortionMapWithOverscan);
+
+	/** Use the input blend parameters to draw multiple displacement maps and blend them together into a final distortion displacement map */
+	void BlendDisplacementMaps_RenderThread(FRDGBuilder& GraphBuilder, const FDisplacementMapBlendingParams& BlendState, float OverscanFactor, FRDGTextureRef& OutDistortionMapWithOverscan);
+
+	/** Crop the input overscanned distortion map to the original requested resolution */
+	void CropDisplacementMap_RenderThread(FRDGBuilder& GraphBuilder, const FRDGTextureRef& InDistortionMapWithOverscan, FRDGTextureRef& OutDistortionMap);
+
+	/** Invert the input distortion map to generate a matching undistortion map (with no overscan) */
+	void InvertDistortionMap_RenderThread(FRDGBuilder& GraphBuilder, const FRDGTextureRef& InDistortionMap, FRDGTextureRef& OutUndistortionMap);
+
+private:
+	/** Map of cameras to their associated distortion state and blending parameters, used to determine if and how displacement maps should be rendered for a specific view */
+	TMap<TWeakObjectPtr<ACameraActor>, FDisplacementMapBlendingParams> DistortionStateMap;
+
+	/** Critical section to lock access to the distortion state map when potentially being accessed from multiple threads */
+	mutable FCriticalSection DistortionStateMapCriticalSection;
+};
diff --git a/Engine/Plugins/VirtualProduction/CameraCalibrationCore/Source/CameraCalibrationCore/Public/LensFile.h b/Engine/Plugins/VirtualProduction/CameraCalibrationCore/Source/CameraCalibrationCore/Public/LensFile.h
--- a/Engine/Plugins/VirtualProduction/CameraCalibrationCore/Source/CameraCalibrationCore/Public/LensFile.h
+++ b/Engine/Plugins/VirtualProduction/CameraCalibrationCore/Source/CameraCalibrationCore/Public/LensFile.h
@@ -27,7 +27,7 @@ class UCineCameraComponent;
 class ULensDistortionModelHandlerBase;
 class UTextureRenderTarget2D;
 struct FBaseLensTable;
-
+struct FDisplacementMapBlendingParams;
 
 DECLARE_MULTICAST_DELEGATE_OneParam(FOnLensFileModelChanged, const TSubclassOf<ULensModel>&);
 
@@ -241,6 +241,9 @@ public:
 	/** Returns the delegate that is triggered when the LensModel changes */
 	FOnLensFileModelChanged& OnLensFileModelChanged() { return OnLensFileModelChangedDelegate; }
 
+	/** Returns the distortion state and blend paramters for input focus and zoom */
+	void GetBlendState(float InFocus, float InZoom, FVector2D InFilmback, FDisplacementMapBlendingParams& OutBlendState);
+
 protected:
 	/** Updates derived data entries to make sure it matches what is assigned in map points based on data mode */
 	void UpdateDerivedData();
diff --git a/Engine/Plugins/VirtualProduction/CameraCalibrationCore/Source/CameraCalibrationCore/Public/LensFileRendering.h b/Engine/Plugins/VirtualProduction/CameraCalibrationCore/Source/CameraCalibrationCore/Public/LensFileRendering.h
new file mode 100644
--- /dev/null
+++ b/Engine/Plugins/VirtualProduction/CameraCalibrationCore/Source/CameraCalibrationCore/Public/LensFileRendering.h
@@ -0,0 +1,90 @@
+// Copyright Epic Games, Inc. All Rights Reserved.
+
+#pragma once
+
+#include "CoreTypes.h"
+#include "Math/Vector2D.h"
+#include "Math/Vector4.h"
+
+#include "LensDistortionModelHandlerBase.h"
+
+class UTextureRenderTarget2D;
+
+/** Types of blending used which drives permutation */
+enum class EDisplacementMapBlendType :uint8
+{
+	OneFocusOneZoom, //no blending
+	OneFocusTwoZoom, //Bezier interp between two zoom points
+	TwoFocusOneZoom, //Linear interp between two focus points
+	TwoFocusTwoZoom, //Two Besier interp between each pair of zoom points and one linear interp across focus
+};
+
+/** Corner of the blending patch which stores x and y values and x and y tangents */
+struct FDisplacementMapBlendPatchCorner
+{
+	FDisplacementMapBlendPatchCorner(float InX, float InY, float InTangentX, float InTangentY)
+		: X(InX)
+		, Y(InY)
+		, TangentX(InTangentX)
+		, TangentY(InTangentY)
+	{ }
+
+	FDisplacementMapBlendPatchCorner()
+		: FDisplacementMapBlendPatchCorner(0.0f, 0.0f, 0.0f, 0.0f)
+	{ }
+	
+	/** Converts the patch corner values to a vector */
+	FVector4f ToVector() const { return FVector4f(X, Y, TangentX, TangentY); }
+	
+	/** X coordinate of the corner */
+	float X;
+
+	/** Y coordinate of the corner */
+	float Y;
+
+	/** Tangent in the x direction of the corner */
+	float TangentX;
+
+	/** Tangent in the y direction of the corner */
+	float TangentY;
+};
+
+/** Single struct containing blending params for all types */
+struct FDisplacementMapBlendingParams
+{
+	/** Active type of blending */
+	EDisplacementMapBlendType BlendType = EDisplacementMapBlendType::OneFocusOneZoom;
+
+	/** Bezier blend parameters */
+	float EvalFocus = 0.0f;
+	float EvalZoom = 0.0f;
+
+	/** Corners of the blending patch, indexed in the following order: (X0, Y0) -> (X1, Y0) -> (X1, Y1) -> (X0, Y1) */
+	FDisplacementMapBlendPatchCorner PatchCorners[4];
+	
+	/** Distortion state for each of four possible corners to be blended */
+	FLensDistortionState States[4];
+
+	/** Scale parameter that allows displacement maps for one sensor size to be applied to camera's with a different sensor size */
+	FVector2D FxFyScale = { 1.0f, 1.0f };
+
+	/** Image center parameter to compute center shift needed to offset resulting map */
+	FVector2D PrincipalPoint = { 0.5f, 0.5f };
+};
+
+namespace LensFileRendering
+{
+	/** Clears the given render target. Useful when no distortion can be applied and the RT has to be resetted */
+	void ClearDisplacementMap(UTextureRenderTarget2D* OutRenderTarget);
+
+	/**
+	 * Draws the blended result of displacement map from input textures based on blend parameters
+	 * One texture is always needed to do a passthrough. Up to four textures can be blended using bilinear
+	 */
+	bool DrawBlendedDisplacementMap(UTextureRenderTarget2D* OutRenderTarget
+		, const FDisplacementMapBlendingParams& BlendParams
+		, UTextureRenderTarget2D* SourceTextureOne
+		, UTextureRenderTarget2D* SourceTextureTwo = nullptr
+		, UTextureRenderTarget2D* SourceTextureThree = nullptr
+		, UTextureRenderTarget2D* SourceTextureFour = nullptr);
+}
diff --git a/Engine/Plugins/VirtualProduction/LensComponent/Source/LensComponent/Private/LensComponent.cpp b/Engine/Plugins/VirtualProduction/LensComponent/Source/LensComponent/Private/LensComponent.cpp
--- a/Engine/Plugins/VirtualProduction/LensComponent/Source/LensComponent/Private/LensComponent.cpp
+++ b/Engine/Plugins/VirtualProduction/LensComponent/Source/LensComponent/Private/LensComponent.cpp
@@ -2,6 +2,8 @@
 
 #include "LensComponent.h"
 
+#include "Camera/CameraActor.h"
+#include "CameraCalibrationSettings.h"
 #include "CameraCalibrationSubsystem.h"
 #include "CineCameraComponent.h"
 #include "Controllers/LiveLinkTransformController.h"
@@ -115,6 +117,8 @@ void ULensComponent::TickComponent(float DeltaTime, ELevelTick TickType, FActorC
 		bWasDistortionEvaluated = false;
 		TObjectPtr<ULensDistortionModelHandlerBase> LensDistortionHandler = LensDistortionHandlerMap.FindRef(LensModel);
 
+		FDisplacementMapBlendingParams BlendState;
+
 		if (LensDistortionHandler)
 		{
 			switch (DistortionStateSource)
@@ -130,6 +134,8 @@ void ULensComponent::TickComponent(float DeltaTime, ELevelTick TickType, FActorC
 
 				DistortionState = LensDistortionHandler->GetCurrentDistortionState();
 
+				LensFile->GetBlendState(EvalInputs.Focus, EvalInputs.Zoom, FVector2D(EvalInputs.Filmback.SensorWidth, EvalInputs.Filmback.SensorHeight), BlendState);
+
 				// Adjust overscan by the overscan multiplier
 				if (bScaleOverscan)
 				{
@@ -150,6 +156,9 @@ void ULensComponent::TickComponent(float DeltaTime, ELevelTick TickType, FActorC
 				LensDistortionHandler->SetDistortionState(DistortionState);
 				LensDistortionHandler->SetCameraFilmback(CineCameraComponent->Filmback);
 
+				BlendState.States[0] = DistortionState;
+				BlendState.BlendType = EDisplacementMapBlendType::OneFocusOneZoom;
+
 				//Recompute overscan factor for the distortion state
 				float OverscanFactor = LensDistortionHandler->ComputeOverscanFactor();
 
@@ -177,30 +186,45 @@ void ULensComponent::TickComponent(float DeltaTime, ELevelTick TickType, FActorC
 		{
 			if (LensDistortionHandler && bWasDistortionEvaluated)
 			{
-				// Get the current distortion MID from the lens distortion handler
-				UMaterialInstanceDynamic* NewDistortionMID = LensDistortionHandler->GetDistortionMID();
-
-				// If the MID has changed
-				if (LastDistortionMID != NewDistortionMID)
+				if (DistortionRenderingMode == EDistortionRenderingMode::LegacyPPM)
 				{
-					CineCameraComponent->RemoveBlendable(LastDistortionMID);
-					CineCameraComponent->AddOrUpdateBlendable(NewDistortionMID);
+					// Get the current distortion MID from the lens distortion handler
+					UMaterialInstanceDynamic* NewDistortionMID = LensDistortionHandler->GetDistortionMID();
+
+					// If the MID has changed
+					if (LastDistortionMID != NewDistortionMID)
+					{
+						CineCameraComponent->RemoveBlendable(LastDistortionMID);
+						CineCameraComponent->AddOrUpdateBlendable(NewDistortionMID);
+					}
+
+					// Cache the latest distortion MID
+					LastDistortionMID = NewDistortionMID;
+
+					// Get the overscan factor and use it to modify the target camera's FOV
+					const float OverscanFactor = LensDistortionHandler->GetOverscanFactor();
+					const float OverscanSensorWidth = GetDesqueezedSensorWidth(CineCameraComponent) * OverscanFactor;
+					const float OverscanFOV = FMath::RadiansToDegrees(2.0f * FMath::Atan(OverscanSensorWidth / (2.0f * OriginalFocalLength)));
+					CineCameraComponent->SetFieldOfView(OverscanFOV);
+
+					// Update the minimum and maximum focal length of the camera (if needed)
+					CineCameraComponent->LensSettings.MinFocalLength = FMath::Min(CineCameraComponent->LensSettings.MinFocalLength, CineCameraComponent->CurrentFocalLength);
+					CineCameraComponent->LensSettings.MaxFocalLength = FMath::Max(CineCameraComponent->LensSettings.MaxFocalLength, CineCameraComponent->CurrentFocalLength);
+
+					bIsDistortionSetup = true;
+				}
+				else if (DistortionRenderingMode == EDistortionRenderingMode::SceneViewExtension)
+				{
+					if (UCameraCalibrationSubsystem* SubSystem = GEngine->GetEngineSubsystem<UCameraCalibrationSubsystem>())
+					{
+						if (ACameraActor* CameraActor = Cast<ACameraActor>(GetOwner()))
+						{
+							SubSystem->SetLensDistortionSVEState(CameraActor, BlendState);
+						}
+					}
+
+					bIsDistortionSetup = true;
 				}
-
-				// Cache the latest distortion MID
-				LastDistortionMID = NewDistortionMID;
-
-				// Get the overscan factor and use it to modify the target camera's FOV
-				const float OverscanFactor = LensDistortionHandler->GetOverscanFactor();
-				const float OverscanSensorWidth = GetDesqueezedSensorWidth(CineCameraComponent) * OverscanFactor;
-				const float OverscanFOV = FMath::RadiansToDegrees(2.0f * FMath::Atan(OverscanSensorWidth / (2.0f * OriginalFocalLength)));
-				CineCameraComponent->SetFieldOfView(OverscanFOV);
-
-				// Update the minimum and maximum focal length of the camera (if needed)
-				CineCameraComponent->LensSettings.MinFocalLength = FMath::Min(CineCameraComponent->LensSettings.MinFocalLength, CineCameraComponent->CurrentFocalLength);
-				CineCameraComponent->LensSettings.MaxFocalLength = FMath::Max(CineCameraComponent->LensSettings.MaxFocalLength, CineCameraComponent->CurrentFocalLength);
-
-				bIsDistortionSetup = true;
 			}
 			else
 			{
@@ -311,6 +335,13 @@ void ULensComponent::PostEditChangeProperty(struct FPropertyChangedEvent& Proper
 			}
 		}
 	}
+	else if (PropertyName == GET_MEMBER_NAME_CHECKED(ULensComponent, DistortionRenderingMode))
+	{
+		if (UCineCameraComponent* const CineCameraComponent = Cast<UCineCameraComponent>(TargetCameraComponent.GetComponent(GetOwner())))
+		{
+			CleanupDistortion(CineCameraComponent);
+		}
+	}
 	else if (PropertyName == GET_MEMBER_NAME_CHECKED(ULensComponent, TargetCameraComponent))
 	{
 		// Clean up distortion on the last target camera
@@ -744,6 +775,14 @@ void ULensComponent::CleanupDistortion(UCineCameraComponent* const CineCameraCom
 		// Update the minimum and maximum focal length of the camera (if needed)
 		CineCameraComponent->LensSettings.MinFocalLength = FMath::Min(CineCameraComponent->LensSettings.MinFocalLength, CineCameraComponent->CurrentFocalLength);
 		CineCameraComponent->LensSettings.MaxFocalLength = FMath::Max(CineCameraComponent->LensSettings.MaxFocalLength, CineCameraComponent->CurrentFocalLength);
+
+		if (UCameraCalibrationSubsystem* SubSystem = GEngine->GetEngineSubsystem<UCameraCalibrationSubsystem>())
+		{
+			if (ACameraActor* CameraActor = Cast<ACameraActor>(GetOwner()))
+			{
+				SubSystem->ClearLensDistortionSVEState(CameraActor);
+			}
+		}
 	}
 
 	bIsDistortionSetup = false;
diff --git a/Engine/Plugins/VirtualProduction/LensComponent/Source/LensComponent/Public/LensComponent.h b/Engine/Plugins/VirtualProduction/LensComponent/Source/LensComponent/Public/LensComponent.h
--- a/Engine/Plugins/VirtualProduction/LensComponent/Source/LensComponent/Public/LensComponent.h
+++ b/Engine/Plugins/VirtualProduction/LensComponent/Source/LensComponent/Public/LensComponent.h
@@ -55,6 +55,16 @@ enum class EDistortionSource : uint8
 	Manual,
 };
 
+/** Specifies how the distortion should be rendered in the post-processing pipeline */
+UENUM(BlueprintType)
+enum class EDistortionRenderingMode : uint8
+{
+	/** Use the legacy post process material */
+	LegacyPPM,
+	/** Use the experimental lens distortion scene view extension. Further control of where distortion is rendered can be set via the console command r.TSR.LensDistortion */
+	SceneViewExtension UMETA(DisplayName = "Scene View Extension (Experimental)"),
+};
+
 /** Component for applying a post-process lens distortion effect to a CineCameraComponent on the same actor */
 UCLASS(HideCategories=(Tags, Activation, Cooking, AssetUserData, Collision), meta=(BlueprintSpawnableComponent))
 class LENSCOMPONENT_API ULensComponent : public UActorComponent
@@ -283,6 +293,10 @@ protected:
 	UPROPERTY(EditAnywhere, BlueprintReadWrite, Getter = ShouldApplyDistortion, Setter = SetApplyDistortion, Category = "Distortion")
 	bool bApplyDistortion = false;
 
+	/** Specifies how the distortion should be rendered in the post-processing pipeline */
+	UPROPERTY(EditAnywhere, Category = "Distortion", meta = (EditCondition = "bApplyDistortion"))
+	EDistortionRenderingMode DistortionRenderingMode = EDistortionRenderingMode::LegacyPPM;
+
 	/** The current lens model used for distortion */
 	UPROPERTY(EditAnywhere, Category = "Distortion", meta = (EditCondition = "DistortionStateSource == EDistortionSource::Manual"))
 	TSubclassOf<ULensModel> LensModel;
@@ -292,7 +306,7 @@ protected:
 	FLensDistortionState DistortionState;
 
 	/** Whether to scale the computed overscan by the overscan percentage */
-	UPROPERTY(AdvancedDisplay, BlueprintReadWrite, Category = "Distortion")
+	UPROPERTY(AdvancedDisplay, BlueprintReadWrite, Category = "Distortion", meta = (InlineEditConditionToggle))
 	bool bScaleOverscan = false;
 
 	/** The percentage of the computed overscan that should be applied to the target camera */
diff --git a/Engine/Source/Runtime/Renderer/Internal/PostProcess/LensDistortion.h b/Engine/Source/Runtime/Renderer/Internal/PostProcess/LensDistortion.h
--- a/Engine/Source/Runtime/Renderer/Internal/PostProcess/LensDistortion.h
+++ b/Engine/Source/Runtime/Renderer/Internal/PostProcess/LensDistortion.h
@@ -91,5 +91,14 @@ namespace LensDistortion
 	 */
 	RENDERER_API const FLensDistortionLUT& GetLUTUnsafe(const FSceneView& InView);
 
+	/**
+	 * Set the view lens distortion LUT.
+	 * Unsafe to due internal argument cast from "FSceneView" to "FViewInfo".
+	 *
+	 * @param InView Active view info.
+	 * @param DistortionLUT The LUT object to set on the FViewInfo
+	 */
+	RENDERER_API void SetLUTUnsafe(FSceneView& InView, const FLensDistortionLUT& DistortionLUT);
+
 } // end namespace LensDistortion
 
diff --git a/Engine/Source/Runtime/Renderer/Private/PostProcess/LensDistortion.cpp b/Engine/Source/Runtime/Renderer/Private/PostProcess/LensDistortion.cpp
--- a/Engine/Source/Runtime/Renderer/Private/PostProcess/LensDistortion.cpp
+++ b/Engine/Source/Runtime/Renderer/Private/PostProcess/LensDistortion.cpp
@@ -248,3 +248,11 @@ const FLensDistortionLUT& LensDistortion::GetLUTUnsafe(const FSceneView& InView)
 	return static_cast<const FViewInfo&>(InView).LensDistortionLUT;
 }
 
+void LensDistortion::SetLUTUnsafe(FSceneView& InView, const FLensDistortionLUT& DistortionLUT)
+{
+	check(InView.bIsViewInfo);
+
+	FViewInfo& ViewInfo = static_cast<FViewInfo&>(InView);
+
+	ViewInfo.LensDistortionLUT = DistortionLUT;
+}
