CameraCalibration: Fix ST Map processing bug, where undistortion map was used instead of distortion map for computing overscan

diff --git a/Engine/Plugins/VirtualProduction/CameraCalibrationCore/Source/CameraCalibrationCore/Private/LensFile.cpp b/Engine/Plugins/VirtualProduction/CameraCalibrationCore/Source/CameraCalibrationCore/Private/LensFile.cpp
--- a/Engine/Plugins/VirtualProduction/CameraCalibrationCore/Source/CameraCalibrationCore/Private/LensFile.cpp
+++ b/Engine/Plugins/VirtualProduction/CameraCalibrationCore/Source/CameraCalibrationCore/Private/LensFile.cpp
@@ -590,7 +590,7 @@ bool ULensFile::EvaluateDistortionForSTMaps(float InFocus, float InZoom, FVector
 			if (const FSTMapZoomPoint* ZoomPoint = FocusPoint.GetZoomPoint(FocusCurve.Zoom))
 			{
 				OutUndistortedMap = ZoomPoint->DerivedDistortionData.UndistortionDisplacementMap;
-				OutDistortedMap = ZoomPoint->DerivedDistortionData.UndistortionDisplacementMap;
+				OutDistortedMap = ZoomPoint->DerivedDistortionData.DistortionDisplacementMap;
 			}
 		});
 
