CameraCalibration: Add overscan support to lens distortion scene view extension.

diff --git a/Engine/Plugins/VirtualProduction/CameraCalibrationCore/Shaders/Private/DrawDisplacementMaps.usf b/Engine/Plugins/VirtualProduction/CameraCalibrationCore/Shaders/Private/DrawDisplacementMaps.usf
--- a/Engine/Plugins/VirtualProduction/CameraCalibrationCore/Shaders/Private/DrawDisplacementMaps.usf
+++ b/Engine/Plugins/VirtualProduction/CameraCalibrationCore/Shaders/Private/DrawDisplacementMaps.usf
@@ -16,7 +16,8 @@ float K3;
 float P1;
 float P2;
 
-float OverscanFactor;
+float InverseOverscan;
+float CameraOverscan;
 
 RWTexture2D<float2> OutDistortionMap;
 
@@ -60,8 +61,10 @@ void MainCS(uint3 DispatchThreadId : SV_DispatchThreadID)
 	float2 UndistortedUV = (DispatchThreadId.xy + 0.5) * ThreadIdToUV;
 	
 	// Rescale the UV by an overscan factor (>= 1.0) to widen the range of UVs to be distorted
-	UndistortedUV = ((UndistortedUV - 0.5) * OverscanFactor) + 0.5;
+	UndistortedUV = ((UndistortedUV - 0.5) * InverseOverscan * CameraOverscan) + 0.5;
 	
 	float2 DistortedUV = SphericalDistort(UndistortedUV);
-	OutDistortionMap[DispatchThreadId.xy] = DistortedUV - UndistortedUV;
+	
+	// Compute the UV displacement and scale it by the CameraOverscan to account for the increased resolution of the overscanned scene texture 
+	OutDistortionMap[DispatchThreadId.xy] = (DistortedUV - UndistortedUV) / CameraOverscan;
 }
diff --git a/Engine/Plugins/VirtualProduction/CameraCalibrationCore/Source/CameraCalibrationCore/Private/LensDistortionSceneViewExtension.cpp b/Engine/Plugins/VirtualProduction/CameraCalibrationCore/Source/CameraCalibrationCore/Private/LensDistortionSceneViewExtension.cpp
--- a/Engine/Plugins/VirtualProduction/CameraCalibrationCore/Source/CameraCalibrationCore/Private/LensDistortionSceneViewExtension.cpp
+++ b/Engine/Plugins/VirtualProduction/CameraCalibrationCore/Source/CameraCalibrationCore/Private/LensDistortionSceneViewExtension.cpp
@@ -3,6 +3,7 @@
 #include "LensDistortionSceneViewExtension.h"
 
 #include "Camera/CameraActor.h"
+#include "Camera/CameraComponent.h"
 #include "Containers/DynamicRHIResourceArray.h"
 #include "GlobalShader.h"
 #include "PostProcess/LensDistortion.h"
@@ -55,7 +56,8 @@ class FDrawDistortionDisplacementMapCS : public FGlobalShader
 		SHADER_PARAMETER(float, K3)
 		SHADER_PARAMETER(float, P1)
 		SHADER_PARAMETER(float, P2)
-		SHADER_PARAMETER(float, OverscanFactor)
+		SHADER_PARAMETER(float, InverseOverscan)
+		SHADER_PARAMETER(float, CameraOverscan)
 		SHADER_PARAMETER_RDG_TEXTURE_UAV(RWTexture2D, OutDistortionMap)
 	END_SHADER_PARAMETER_STRUCT()
 };
@@ -137,7 +139,7 @@ public:
 IMPLEMENT_GLOBAL_SHADER(FInvertDisplacementVS, "/Plugin/CameraCalibrationCore/Private/InvertDisplacementMap.usf", "MainVS", SF_Vertex);
 IMPLEMENT_GLOBAL_SHADER(FInvertDisplacementPS, "/Plugin/CameraCalibrationCore/Private/InvertDisplacementMap.usf", "MainPS", SF_Pixel);
 
-void FLensDistortionSceneViewExtension::DrawDisplacementMap_RenderThread(FRDGBuilder& GraphBuilder, const FLensDistortionState& CurrentState, float OverscanFactor, FRDGTextureRef& OutDistortionMapWithOverscan)
+void FLensDistortionSceneViewExtension::DrawDisplacementMap_RenderThread(FRDGBuilder& GraphBuilder, const FLensDistortionState& CurrentState, float InverseOverscan, float CameraOverscan, FRDGTextureRef& OutDistortionMapWithOverscan)
 {
 	if (CurrentState.DistortionInfo.Parameters.IsEmpty())
 	{
@@ -161,7 +163,8 @@ void FLensDistortionSceneViewExtension::DrawDisplacementMap_RenderThread(FRDGBui
 	PassParameters->P1 = CurrentState.DistortionInfo.Parameters[3];
 	PassParameters->P2 = CurrentState.DistortionInfo.Parameters[4];
 
-	PassParameters->OverscanFactor = OverscanFactor;
+	PassParameters->InverseOverscan = InverseOverscan;
+	PassParameters->CameraOverscan = CameraOverscan;
 
 	TShaderMapRef<FDrawDistortionDisplacementMapCS> ComputeShader(GetGlobalShaderMap(GMaxRHIFeatureLevel));
 	FComputeShaderUtils::AddPass(
@@ -191,14 +194,14 @@ void FLensDistortionSceneViewExtension::CropDisplacementMap_RenderThread(FRDGBui
 		FIntVector(FMath::DivideAndRoundUp(LUTResolution.X, 8), FMath::DivideAndRoundUp(LUTResolution.Y, 8), 1));
 }
 
-void FLensDistortionSceneViewExtension::BlendDisplacementMaps_RenderThread(FRDGBuilder& GraphBuilder, const FDisplacementMapBlendingParams& BlendState, float OverscanFactor, FRDGTextureRef& OutDistortionMapWithOverscan)
+void FLensDistortionSceneViewExtension::BlendDisplacementMaps_RenderThread(FRDGBuilder& GraphBuilder, const FDisplacementMapBlendingParams& BlendState, float InverseOverscan, float CameraOverscan, FRDGTextureRef& OutDistortionMapWithOverscan)
 {
 	FBlendDistortionDisplacementMapCS::FParameters* PassParameters = GraphBuilder.AllocParameters<FBlendDistortionDisplacementMapCS::FParameters>();
 
 	// Draw the first distortion map, which should always be valid
 	{
 		FRDGTextureRef Distortion1 = GraphBuilder.CreateTexture(OutDistortionMapWithOverscan->Desc, TEXT("DistortingDisplacement1"));
-		DrawDisplacementMap_RenderThread(GraphBuilder, BlendState.States[0], OverscanFactor, Distortion1);
+		DrawDisplacementMap_RenderThread(GraphBuilder, BlendState.States[0], InverseOverscan, CameraOverscan, Distortion1);
 		PassParameters->InputDistortionMap1 = GraphBuilder.CreateSRV(Distortion1);
 	}
 
@@ -206,7 +209,7 @@ void FLensDistortionSceneViewExtension::BlendDisplacementMaps_RenderThread(FRDGB
 	if (BlendState.BlendType != EDisplacementMapBlendType::OneFocusOneZoom)
 	{
 		FRDGTextureRef Distortion2 = GraphBuilder.CreateTexture(OutDistortionMapWithOverscan->Desc, TEXT("DistortingDisplacement2"));
-		DrawDisplacementMap_RenderThread(GraphBuilder, BlendState.States[1], OverscanFactor, Distortion2);
+		DrawDisplacementMap_RenderThread(GraphBuilder, BlendState.States[1], InverseOverscan, CameraOverscan, Distortion2);
 		PassParameters->InputDistortionMap2 = GraphBuilder.CreateSRV(Distortion2);
 	}
 
@@ -216,8 +219,8 @@ void FLensDistortionSceneViewExtension::BlendDisplacementMaps_RenderThread(FRDGB
 		FRDGTextureRef Distortion3 = GraphBuilder.CreateTexture(OutDistortionMapWithOverscan->Desc, TEXT("DistortingDisplacement3"));
 		FRDGTextureRef Distortion4 = GraphBuilder.CreateTexture(OutDistortionMapWithOverscan->Desc, TEXT("DistortingDisplacement4"));
 
-		DrawDisplacementMap_RenderThread(GraphBuilder, BlendState.States[2], OverscanFactor, Distortion3);
-		DrawDisplacementMap_RenderThread(GraphBuilder, BlendState.States[3], OverscanFactor, Distortion4);
+		DrawDisplacementMap_RenderThread(GraphBuilder, BlendState.States[2], InverseOverscan, CameraOverscan, Distortion3);
+		DrawDisplacementMap_RenderThread(GraphBuilder, BlendState.States[3], InverseOverscan, CameraOverscan, Distortion4);
 
 		PassParameters->InputDistortionMap3 = GraphBuilder.CreateSRV(Distortion3);
 		PassParameters->InputDistortionMap4 = GraphBuilder.CreateSRV(Distortion4);
@@ -398,11 +401,11 @@ void FLensDistortionSceneViewExtension::PreRenderView_RenderThread(FRDGBuilder&
 		ViewDistortionLUT.UndistortingDisplacementTexture = GraphBuilder.CreateTexture(UndistortionMapDesc, TEXT("UndistortionDisplacementMap"));
 
 		// In order to guarantee that we can generate a complete undistortion map, the distortion map we invert needs to have some overscan 
-		float OverscanFactor = FMath::Clamp(CVarLensDistortionDisplacementOverscan.GetValueOnRenderThread(), 1.0f, 2.0f);
+		float InverseOverscan = FMath::Clamp(CVarLensDistortionDisplacementOverscan.GetValueOnRenderThread(), 1.0f, 2.0f);
 
 		// Adjust the overscan resolution to be square, with each side being a multiple of 8
-		const FIntPoint OverscanResolution = FIntPoint(FMath::CeilToInt(OverscanFactor * 32) * 8);
-		OverscanFactor = float(OverscanResolution.X) / float(DisplacementMapResolution.X);
+		const FIntPoint OverscanResolution = FIntPoint(FMath::CeilToInt(InverseOverscan * 32) * 8);
+		InverseOverscan = float(OverscanResolution.X) / float(DisplacementMapResolution.X);
 
 		// Create the texture for the overscanned distortion map
 		FRDGTextureDesc OverscanDesc = FRDGTextureDesc::Create2D(
@@ -413,7 +416,13 @@ void FLensDistortionSceneViewExtension::PreRenderView_RenderThread(FRDGBuilder&
 
 		FRDGTextureRef DistortionMapWithOverscan = GraphBuilder.CreateTexture(OverscanDesc, TEXT("DistortionMapWithOverscan"));
 
-		BlendDisplacementMaps_RenderThread(GraphBuilder, BlendState, OverscanFactor, DistortionMapWithOverscan);
+		float CameraOverscan = 1.0f;
+		if (UCameraComponent* CameraComponent = CameraActor->GetCameraComponent())
+		{
+			CameraOverscan = CameraComponent->Overscan + 1.0f;
+		}
+
+		BlendDisplacementMaps_RenderThread(GraphBuilder, BlendState, InverseOverscan, CameraOverscan, DistortionMapWithOverscan);
 		InvertDistortionMap_RenderThread(GraphBuilder, DistortionMapWithOverscan, ViewDistortionLUT.UndistortingDisplacementTexture);
 		CropDisplacementMap_RenderThread(GraphBuilder, DistortionMapWithOverscan, ViewDistortionLUT.DistortingDisplacementTexture);
 
diff --git a/Engine/Plugins/VirtualProduction/CameraCalibrationCore/Source/CameraCalibrationCore/Public/LensDistortionSceneViewExtension.h b/Engine/Plugins/VirtualProduction/CameraCalibrationCore/Source/CameraCalibrationCore/Public/LensDistortionSceneViewExtension.h
--- a/Engine/Plugins/VirtualProduction/CameraCalibrationCore/Source/CameraCalibrationCore/Public/LensDistortionSceneViewExtension.h
+++ b/Engine/Plugins/VirtualProduction/CameraCalibrationCore/Source/CameraCalibrationCore/Public/LensDistortionSceneViewExtension.h
@@ -32,10 +32,10 @@ public:
 
 private:
 	/** Use the input distortion state to draw a distortion displacement map */
-	void DrawDisplacementMap_RenderThread(FRDGBuilder& GraphBuilder, const FLensDistortionState& CurrentState, float OverscanFactor, FRDGTextureRef& OutDistortionMapWithOverscan);
+	void DrawDisplacementMap_RenderThread(FRDGBuilder& GraphBuilder, const FLensDistortionState& CurrentState, float InverseOverscan, float CameraOverscan, FRDGTextureRef& OutDistortionMapWithOverscan);
 
 	/** Use the input blend parameters to draw multiple displacement maps and blend them together into a final distortion displacement map */
-	void BlendDisplacementMaps_RenderThread(FRDGBuilder& GraphBuilder, const FDisplacementMapBlendingParams& BlendState, float OverscanFactor, FRDGTextureRef& OutDistortionMapWithOverscan);
+	void BlendDisplacementMaps_RenderThread(FRDGBuilder& GraphBuilder, const FDisplacementMapBlendingParams& BlendState, float InverseOverscan, float CameraOverscan, FRDGTextureRef& OutDistortionMapWithOverscan);
 
 	/** Crop the input overscanned distortion map to the original requested resolution */
 	void CropDisplacementMap_RenderThread(FRDGBuilder& GraphBuilder, const FRDGTextureRef& InDistortionMapWithOverscan, FRDGTextureRef& OutDistortionMap);
