[CameraCalibrationCore]
* Moved code one step closer to not using actors render side

diff --git a/Engine/Plugins/VirtualProduction/CameraCalibrationCore/Source/CameraCalibrationCore/Private/LensDistortionSceneViewExtension.cpp b/Engine/Plugins/VirtualProduction/CameraCalibrationCore/Source/CameraCalibrationCore/Private/LensDistortionSceneViewExtension.cpp
--- a/Engine/Plugins/VirtualProduction/CameraCalibrationCore/Source/CameraCalibrationCore/Private/LensDistortionSceneViewExtension.cpp
+++ b/Engine/Plugins/VirtualProduction/CameraCalibrationCore/Source/CameraCalibrationCore/Private/LensDistortionSceneViewExtension.cpp
@@ -33,13 +33,13 @@ FLensDistortionSceneViewExtension::FLensDistortionSceneViewExtension(const FAuto
 void FLensDistortionSceneViewExtension::UpdateDistortionState_AnyThread(ACameraActor* CameraActor, FDisplacementMapBlendingParams DistortionState)
 {
 	FScopeLock ScopeLock(&DistortionStateMapCriticalSection);
-	DistortionStateMap.Add(CameraActor, DistortionState);
+	DistortionStateMap.Add(CameraActor->GetUniqueID(), {DistortionState, CameraActor->GetCameraComponent()});
 }
 
 void FLensDistortionSceneViewExtension::ClearDistortionState_AnyThread(ACameraActor* CameraActor)
 {
 	FScopeLock ScopeLock(&DistortionStateMapCriticalSection);
-	DistortionStateMap.Remove(CameraActor);
+	DistortionStateMap.Remove(CameraActor->GetUniqueID());
 }
 
 class FDrawDistortionDisplacementMapCS : public FGlobalShader
@@ -365,67 +365,58 @@ void FLensDistortionSceneViewExtension::InvertDistortionMap_RenderThread(FRDGBui
 
 void FLensDistortionSceneViewExtension::PreRenderView_RenderThread(FRDGBuilder& GraphBuilder, FSceneView& InView)
 {
-	if (const ACameraActor* CameraActor = Cast<const ACameraActor>(InView.ViewActor.Get()))
+	FScopeLock ScopeLock(&DistortionStateMapCriticalSection);
+	Rec* Rec = DistortionStateMap.Find(InView.ViewActor.ActorUniqueId);
+	if (!Rec)
 	{
-		FDisplacementMapBlendingParams BlendState;
-		{
-			FScopeLock ScopeLock(&DistortionStateMapCriticalSection);
-			if (FDisplacementMapBlendingParams* BlendStatePtr = DistortionStateMap.Find(CameraActor))
-			{
-				BlendState = *BlendStatePtr;
-			}
-			else
-			{
-				return;
-			}
-		}
+		return;
+	}
 
-		// Create the distortion map and undistortion map textures for the FLensDistortionLUT for this frame
-		const FIntPoint DisplacementMapResolution = FIntPoint(256, 256);
+	float CameraOverscan = Rec->Component ? (Rec->Component->Overscan + 1.0f) : 1.0f;
+	FDisplacementMapBlendingParams BlendState = Rec->Params;
 
-		FRDGTextureDesc DistortionMapDesc = FRDGTextureDesc::Create2D(
-			DisplacementMapResolution,
-			PF_G32R32F,
-			FClearValueBinding::None,
-			TexCreate_ShaderResource | TexCreate_UAV);
+	ScopeLock.Unlock();
 
-		FRDGTextureDesc UndistortionMapDesc = FRDGTextureDesc::Create2D(
-			DisplacementMapResolution,
-			PF_G32R32F,
-			FClearValueBinding::Black,
-			TexCreate_ShaderResource | TexCreate_RenderTargetable);
 
-		FLensDistortionLUT ViewDistortionLUT;
+	// Create the distortion map and undistortion map textures for the FLensDistortionLUT for this frame
+	const FIntPoint DisplacementMapResolution = FIntPoint(256, 256);
 
-		ViewDistortionLUT.DistortingDisplacementTexture = GraphBuilder.CreateTexture(DistortionMapDesc, TEXT("DistortionDisplacementMap"));
-		ViewDistortionLUT.UndistortingDisplacementTexture = GraphBuilder.CreateTexture(UndistortionMapDesc, TEXT("UndistortionDisplacementMap"));
+	FRDGTextureDesc DistortionMapDesc = FRDGTextureDesc::Create2D(
+		DisplacementMapResolution,
+		PF_G32R32F,
+		FClearValueBinding::None,
+		TexCreate_ShaderResource | TexCreate_UAV);
 
-		// In order to guarantee that we can generate a complete undistortion map, the distortion map we invert needs to have some overscan 
-		float InverseOverscan = FMath::Clamp(CVarLensDistortionDisplacementOverscan.GetValueOnRenderThread(), 1.0f, 2.0f);
+	FRDGTextureDesc UndistortionMapDesc = FRDGTextureDesc::Create2D(
+		DisplacementMapResolution,
+		PF_G32R32F,
+		FClearValueBinding::Black,
+		TexCreate_ShaderResource | TexCreate_RenderTargetable);
 
-		// Adjust the overscan resolution to be square, with each side being a multiple of 8
-		const FIntPoint OverscanResolution = FIntPoint(FMath::CeilToInt(InverseOverscan * 32) * 8);
-		InverseOverscan = float(OverscanResolution.X) / float(DisplacementMapResolution.X);
+	FLensDistortionLUT ViewDistortionLUT;
 
-		// Create the texture for the overscanned distortion map
-		FRDGTextureDesc OverscanDesc = FRDGTextureDesc::Create2D(
-			OverscanResolution,
-			PF_G32R32F,
-			FClearValueBinding::None,
-			TexCreate_ShaderResource | TexCreate_UAV);
+	ViewDistortionLUT.DistortingDisplacementTexture = GraphBuilder.CreateTexture(DistortionMapDesc, TEXT("DistortionDisplacementMap"));
+	ViewDistortionLUT.UndistortingDisplacementTexture = GraphBuilder.CreateTexture(UndistortionMapDesc, TEXT("UndistortionDisplacementMap"));
 
-		FRDGTextureRef DistortionMapWithOverscan = GraphBuilder.CreateTexture(OverscanDesc, TEXT("DistortionMapWithOverscan"));
+	// In order to guarantee that we can generate a complete undistortion map, the distortion map we invert needs to have some overscan 
+	float InverseOverscan = FMath::Clamp(CVarLensDistortionDisplacementOverscan.GetValueOnRenderThread(), 1.0f, 2.0f);
 
-		float CameraOverscan = 1.0f;
-		if (UCameraComponent* CameraComponent = CameraActor->GetCameraComponent())
-		{
-			CameraOverscan = CameraComponent->Overscan + 1.0f;
-		}
+	// Adjust the overscan resolution to be square, with each side being a multiple of 8
+	const FIntPoint OverscanResolution = FIntPoint(FMath::CeilToInt(InverseOverscan * 32) * 8);
+	InverseOverscan = float(OverscanResolution.X) / float(DisplacementMapResolution.X);
 
-		BlendDisplacementMaps_RenderThread(GraphBuilder, BlendState, InverseOverscan, CameraOverscan, DistortionMapWithOverscan);
-		InvertDistortionMap_RenderThread(GraphBuilder, DistortionMapWithOverscan, ViewDistortionLUT.UndistortingDisplacementTexture);
-		CropDisplacementMap_RenderThread(GraphBuilder, DistortionMapWithOverscan, ViewDistortionLUT.DistortingDisplacementTexture);
+	// Create the texture for the overscanned distortion map
+	FRDGTextureDesc OverscanDesc = FRDGTextureDesc::Create2D(
+		OverscanResolution,
+		PF_G32R32F,
+		FClearValueBinding::None,
+		TexCreate_ShaderResource | TexCreate_UAV);
 
-		LensDistortion::SetLUTUnsafe(InView, ViewDistortionLUT);
-	}
+	FRDGTextureRef DistortionMapWithOverscan = GraphBuilder.CreateTexture(OverscanDesc, TEXT("DistortionMapWithOverscan"));
+
+	BlendDisplacementMaps_RenderThread(GraphBuilder, BlendState, InverseOverscan, CameraOverscan, DistortionMapWithOverscan);
+	InvertDistortionMap_RenderThread(GraphBuilder, DistortionMapWithOverscan, ViewDistortionLUT.UndistortingDisplacementTexture);
+	CropDisplacementMap_RenderThread(GraphBuilder, DistortionMapWithOverscan, ViewDistortionLUT.DistortingDisplacementTexture);
+
+	LensDistortion::SetLUTUnsafe(InView, ViewDistortionLUT);
 }
diff --git a/Engine/Plugins/VirtualProduction/CameraCalibrationCore/Source/CameraCalibrationCore/Public/LensDistortionSceneViewExtension.h b/Engine/Plugins/VirtualProduction/CameraCalibrationCore/Source/CameraCalibrationCore/Public/LensDistortionSceneViewExtension.h
--- a/Engine/Plugins/VirtualProduction/CameraCalibrationCore/Source/CameraCalibrationCore/Public/LensDistortionSceneViewExtension.h
+++ b/Engine/Plugins/VirtualProduction/CameraCalibrationCore/Source/CameraCalibrationCore/Public/LensDistortionSceneViewExtension.h
@@ -6,6 +6,8 @@
 
 #include "LensFileRendering.h"
 
+class UCameraComponent;
+
 /**
  * View extension drawing distortion/undistortion displacement maps
  */
@@ -45,7 +47,8 @@ private:
 
 private:
 	/** Map of cameras to their associated distortion state and blending parameters, used to determine if and how displacement maps should be rendered for a specific view */
-	TMap<TWeakObjectPtr<ACameraActor>, FDisplacementMapBlendingParams> DistortionStateMap;
+	struct Rec { FDisplacementMapBlendingParams Params; UCameraComponent* Component; };
+	TMap<uint32, Rec> DistortionStateMap;
 
 	/** Critical section to lock access to the distortion state map when potentially being accessed from multiple threads */
 	mutable FCriticalSection DistortionStateMapCriticalSection;
