CameraCalibration: Fix LensFile to properly delete empty focus curves

diff --git a/Engine/Plugins/VirtualProduction/CameraCalibrationCore/Source/CameraCalibrationCore/Private/Tables/LensTableUtils.h b/Engine/Plugins/VirtualProduction/CameraCalibrationCore/Source/CameraCalibrationCore/Private/Tables/LensTableUtils.h
--- a/Engine/Plugins/VirtualProduction/CameraCalibrationCore/Source/CameraCalibrationCore/Private/Tables/LensTableUtils.h
+++ b/Engine/Plugins/VirtualProduction/CameraCalibrationCore/Source/CameraCalibrationCore/Private/Tables/LensTableUtils.h
@@ -28,6 +28,8 @@ namespace LensDataTableUtils
 		{
 			Curve.RemovePoint(InFocus, UE_SMALL_NUMBER);
 		}
+
+		InFocusCurves.RemoveAll([](FocusCurveType InCurve) { return InCurve.IsEmpty(); });
 	}
 	
 	/** Changes the value of a focus point in the container */
@@ -323,6 +325,7 @@ namespace LensDataTableUtils
 	void EmptyTable(Type& InTable)
 	{
 		InTable.FocusPoints.Empty(0);
+		InTable.FocusCurves.Empty(0);
 	}
 
 	struct FPointNeighbors
