Lens Calibration: Fixed issues cuasing artifacts for lens distortion in TSR. Fixed warp mesh overlapping itself for large distortions. Increased warp grid density when computing undistortion map to ensure an accurate approximation to prevent artifacts from manifesting in TSR when it undistorts a history frame

diff --git a/Engine/Plugins/VirtualProduction/CameraCalibrationCore/Shaders/Private/DrawDisplacementMaps.usf b/Engine/Plugins/VirtualProduction/CameraCalibrationCore/Shaders/Private/DrawDisplacementMaps.usf
--- a/Engine/Plugins/VirtualProduction/CameraCalibrationCore/Shaders/Private/DrawDisplacementMaps.usf
+++ b/Engine/Plugins/VirtualProduction/CameraCalibrationCore/Shaders/Private/DrawDisplacementMaps.usf
@@ -54,6 +54,30 @@ float2 SphericalDistort(float2 InputUV)
 	return result;
 }
 
+bool IsDistortionWellOriented(float2 InputUV, float2 DistortedUV, float2 UVOffset)
+{
+	float2 DistortedXOffset = SphericalDistort(InputUV + UVOffset * int2(1, 0));
+	float2 DistortedYOffset = SphericalDistort(InputUV + UVOffset * int2(0, 1));
+
+	float2 XOffset = DistortedXOffset - DistortedUV;
+	float2 YOffset = DistortedYOffset - DistortedUV;
+	
+	bool bIsFrontFacing = (XOffset.x * YOffset.y - XOffset.y * YOffset.x) >= 0;
+	bool bIsRightSideUp = (XOffset.x * UVOffset.x) >= 0 && (YOffset.y * UVOffset.y) >= 0;
+	return bIsFrontFacing && bIsRightSideUp;
+}
+
+float2 TrySphericalDistort(float2 InputUV)
+{
+	float2 DistortedPosition = SphericalDistort(InputUV);
+	float2 UVOffset = 0.5 * ThreadIdToUV;
+	bool bIsOrientedWell =
+		IsDistortionWellOriented(InputUV, DistortedPosition, UVOffset) &&
+		IsDistortionWellOriented(InputUV, DistortedPosition, -UVOffset);
+
+	return lerp(InputUV, DistortedPosition, bIsOrientedWell);
+}
+
 [numthreads(8, 8, 1)]
 void MainCS(uint3 DispatchThreadId : SV_DispatchThreadID)
 {
@@ -63,7 +87,7 @@ void MainCS(uint3 DispatchThreadId : SV_DispatchThreadID)
 	// Rescale the UV by an overscan factor (>= 1.0) to widen the range of UVs to be distorted
 	UndistortedUV = ((UndistortedUV - 0.5) * InverseOverscan * CameraOverscan) + 0.5;
 	
-	float2 DistortedUV = SphericalDistort(UndistortedUV);
+	float2 DistortedUV = TrySphericalDistort(UndistortedUV);
 	
 	// Compute the UV displacement and scale it by the CameraOverscan to account for the increased resolution of the overscanned scene texture 
 	OutDistortionMap[DispatchThreadId.xy] = (DistortedUV - UndistortedUV) / CameraOverscan;
diff --git a/Engine/Plugins/VirtualProduction/CameraCalibrationCore/Source/CameraCalibrationCore/Private/LensDistortionSceneViewExtension.cpp b/Engine/Plugins/VirtualProduction/CameraCalibrationCore/Source/CameraCalibrationCore/Private/LensDistortionSceneViewExtension.cpp
--- a/Engine/Plugins/VirtualProduction/CameraCalibrationCore/Source/CameraCalibrationCore/Private/LensDistortionSceneViewExtension.cpp
+++ b/Engine/Plugins/VirtualProduction/CameraCalibrationCore/Source/CameraCalibrationCore/Private/LensDistortionSceneViewExtension.cpp
@@ -11,16 +11,9 @@
 #include "ShaderParameterStruct.h"
 #include "SystemTextures.h"
 
-TAutoConsoleVariable<float> CVarLensDistortionDisplacementOverscan(
-	TEXT("r.LensDistortion.DisplacementMapOverscan"),
-	1.25f,
-	TEXT("A factor to scale the distortion displacement map to ensure that the undistortion map is properly invertible.\n")
-	TEXT("Value is clamped between 1.0 (no overscan) and 2.0.\n"),
-	ECVF_RenderThreadSafe);
-
 TAutoConsoleVariable<int32> CVarLensDistortionInvertGridDensity(
 	TEXT("r.LensDistortion.InvertGridDensity"),
-	64,
+	96,
 	TEXT("The number of squares drawn by the shader that inverts the distortion displacement map\n")
 	TEXT("Value is clamped between 64 and 255.\n"),
 	ECVF_RenderThreadSafe);
@@ -277,7 +270,7 @@ void FLensDistortionSceneViewExtension::BlendDisplacementMaps_RenderThread(FRDGB
 		FIntVector(FMath::DivideAndRoundUp(DistortionMapResolution.X, 8), FMath::DivideAndRoundUp(DistortionMapResolution.Y, 8), 1));
 }
 
-void FLensDistortionSceneViewExtension::InvertDistortionMap_RenderThread(FRDGBuilder& GraphBuilder, const FRDGTextureRef& InDistortionMap, FRDGTextureRef& OutUndistortionMap)
+void FLensDistortionSceneViewExtension::InvertDistortionMap_RenderThread(FRDGBuilder& GraphBuilder, const FRDGTextureRef& InDistortionMap, float InInverseOverscan, FRDGTextureRef& OutUndistortionMap)
 {
 	FInvertDisplacementParameters* PassParameters = GraphBuilder.AllocParameters<FInvertDisplacementParameters>();
 
@@ -286,6 +279,14 @@ void FLensDistortionSceneViewExtension::InvertDistortionMap_RenderThread(FRDGBui
 	Output.ViewRect = FIntRect(FIntPoint(0, 0), OutUndistortionMap->Desc.Extent);
 	Output.LoadAction = ERenderTargetLoadAction::EClear;
 	Output.UpdateVisualizeTextureExtent();
+	
+	const int32 NumSquares = FMath::Clamp(CVarLensDistortionInvertGridDensity.GetValueOnRenderThread(), 64, 255);
+	FIntPoint GridDimensions = FIntPoint(NumSquares);
+
+	// Scale the grid density by the overscan to ensure that there is no change in number of distorted vertices
+	GridDimensions = FIntPoint(FMath::CeilToInt(GridDimensions.X * InInverseOverscan), FMath::CeilToInt(GridDimensions.Y * InInverseOverscan));
+	
+	PassParameters->GridDimensions = GridDimensions;
 
 	PassParameters->RenderTargets[0] = Output.GetRenderTargetBinding();
 
@@ -300,12 +301,8 @@ void FLensDistortionSceneViewExtension::InvertDistortionMap_RenderThread(FRDGBui
 		RDG_EVENT_NAME("InvertDistortionDisplacementMap"),
 		PassParameters,
 		ERDGPassFlags::Raster,
-		[PassParameters, Output](FRDGAsyncTask, FRHICommandList& RHICmdList)
+		[PassParameters, GridDimensions, Output](FRDGAsyncTask, FRHICommandList& RHICmdList)
 		{
-			const int32 NumSquares = FMath::Clamp(CVarLensDistortionInvertGridDensity.GetValueOnRenderThread(), 64, 255);
-			const FIntPoint GridDimensions = FIntPoint(NumSquares);
-			PassParameters->GridDimensions = GridDimensions;
-
 			RHICmdList.SetViewport(Output.ViewRect.Min.X, Output.ViewRect.Min.Y, 0.0f, Output.ViewRect.Max.X, Output.ViewRect.Max.Y, 1.0f);
 
 			TShaderMapRef<FInvertDisplacementVS> VertexShader(GetGlobalShaderMap(GMaxRHIFeatureLevel));
@@ -322,7 +319,7 @@ void FLensDistortionSceneViewExtension::InvertDistortionMap_RenderThread(FRDGBui
 			RHICmdList.SetStreamSource(0, NULL, 0);
 
 			// The following code for setting up this index buffer based on FTesselatedScreenRectangleIndexBuffer::InitRHI()
-			TResourceArray<uint16> IndexBuffer;
+			TResourceArray<uint32> IndexBuffer;
 
 			const uint32 Width = GridDimensions.X;
 			const uint32 Height = GridDimensions.Y;
@@ -331,17 +328,17 @@ void FLensDistortionSceneViewExtension::InvertDistortionMap_RenderThread(FRDGBui
 
 			IndexBuffer.AddUninitialized(NumIndices);
 
-			uint16* IndexBufferData = (uint16*)IndexBuffer.GetData();
+			uint32* IndexBufferData = (uint32*)IndexBuffer.GetData();
 
 			for (uint32 IndexY = 0; IndexY < Height; ++IndexY)
 			{
 				for (uint32 IndexX = 0; IndexX < Width; ++IndexX)
 				{
 					// left top to bottom right in reading order
-					uint16 Index00 = IndexX + IndexY * (Width + 1);
-					uint16 Index10 = Index00 + 1;
-					uint16 Index01 = Index00 + (Width + 1);
-					uint16 Index11 = Index01 + 1;
+					uint32 Index00 = IndexX + IndexY * (Width + 1);
+					uint32 Index10 = Index00 + 1;
+					uint32 Index01 = Index00 + (Width + 1);
+					uint32 Index11 = Index01 + 1;
 
 					// triangle A
 					*IndexBufferData++ = Index00;
@@ -357,7 +354,7 @@ void FLensDistortionSceneViewExtension::InvertDistortionMap_RenderThread(FRDGBui
 
 			// Create index buffer. Fill buffer with initial data upon creation
 			FRHIResourceCreateInfo CreateInfo(TEXT("InvertDistortionMapIndexBuffer"), &IndexBuffer);
-			FBufferRHIRef IndexBufferRHI = RHICmdList.CreateIndexBuffer(sizeof(uint16), IndexBuffer.GetResourceDataSize(), BUF_Static, CreateInfo);
+			FBufferRHIRef IndexBufferRHI = RHICmdList.CreateIndexBuffer(sizeof(uint32), IndexBuffer.GetResourceDataSize(), BUF_Static, CreateInfo);
 
 			RHICmdList.DrawIndexedPrimitive(IndexBufferRHI, 0, 0, NumIndices, 0, NumTriangles, 1);
 		});
@@ -399,11 +396,9 @@ void FLensDistortionSceneViewExtension::PreRenderView_RenderThread(FRDGBuilder&
 	ViewDistortionLUT.UndistortingDisplacementTexture = GraphBuilder.CreateTexture(UndistortionMapDesc, TEXT("UndistortionDisplacementMap"));
 
 	// In order to guarantee that we can generate a complete undistortion map, the distortion map we invert needs to have some overscan 
-	float InverseOverscan = FMath::Clamp(CVarLensDistortionDisplacementOverscan.GetValueOnRenderThread(), 1.0f, 2.0f);
-
 	// Adjust the overscan resolution to be square, with each side being a multiple of 8
-	const FIntPoint OverscanResolution = FIntPoint(FMath::CeilToInt(InverseOverscan * 32) * 8);
-	InverseOverscan = float(OverscanResolution.X) / float(DisplacementMapResolution.X);
+	const FIntPoint OverscanResolution = FIntPoint(FMath::CeilToInt(CameraOverscan * 32) * 8);
+	const float InverseOverscan = float(OverscanResolution.X) / float(DisplacementMapResolution.X);
 
 	// Create the texture for the overscanned distortion map
 	FRDGTextureDesc OverscanDesc = FRDGTextureDesc::Create2D(
@@ -415,7 +410,7 @@ void FLensDistortionSceneViewExtension::PreRenderView_RenderThread(FRDGBuilder&
 	FRDGTextureRef DistortionMapWithOverscan = GraphBuilder.CreateTexture(OverscanDesc, TEXT("DistortionMapWithOverscan"));
 
 	BlendDisplacementMaps_RenderThread(GraphBuilder, BlendState, InverseOverscan, CameraOverscan, DistortionMapWithOverscan);
-	InvertDistortionMap_RenderThread(GraphBuilder, DistortionMapWithOverscan, ViewDistortionLUT.UndistortingDisplacementTexture);
+	InvertDistortionMap_RenderThread(GraphBuilder, DistortionMapWithOverscan, InverseOverscan, ViewDistortionLUT.UndistortingDisplacementTexture);
 	CropDisplacementMap_RenderThread(GraphBuilder, DistortionMapWithOverscan, ViewDistortionLUT.DistortingDisplacementTexture);
 
 	LensDistortion::SetLUTUnsafe(InView, ViewDistortionLUT);
diff --git a/Engine/Plugins/VirtualProduction/CameraCalibrationCore/Source/CameraCalibrationCore/Public/LensDistortionSceneViewExtension.h b/Engine/Plugins/VirtualProduction/CameraCalibrationCore/Source/CameraCalibrationCore/Public/LensDistortionSceneViewExtension.h
--- a/Engine/Plugins/VirtualProduction/CameraCalibrationCore/Source/CameraCalibrationCore/Public/LensDistortionSceneViewExtension.h
+++ b/Engine/Plugins/VirtualProduction/CameraCalibrationCore/Source/CameraCalibrationCore/Public/LensDistortionSceneViewExtension.h
@@ -39,7 +39,7 @@ private:
 	void CropDisplacementMap_RenderThread(FRDGBuilder& GraphBuilder, const FRDGTextureRef& InDistortionMapWithOverscan, FRDGTextureRef& OutDistortionMap);
 
 	/** Invert the input distortion map to generate a matching undistortion map (with no overscan) */
-	void InvertDistortionMap_RenderThread(FRDGBuilder& GraphBuilder, const FRDGTextureRef& InDistortionMap, FRDGTextureRef& OutUndistortionMap);
+	void InvertDistortionMap_RenderThread(FRDGBuilder& GraphBuilder, const FRDGTextureRef& InDistortionMap, float InInverseOverscan, FRDGTextureRef& OutUndistortionMap);
 
 private:
 	/** Map of cameras to their associated distortion state and blending parameters, used to determine if and how displacement maps should be rendered for a specific view */
diff --git a/Engine/Plugins/VirtualProduction/LensComponent/Source/LensComponent/Private/LensComponent.cpp b/Engine/Plugins/VirtualProduction/LensComponent/Source/LensComponent/Private/LensComponent.cpp
--- a/Engine/Plugins/VirtualProduction/LensComponent/Source/LensComponent/Private/LensComponent.cpp
+++ b/Engine/Plugins/VirtualProduction/LensComponent/Source/LensComponent/Private/LensComponent.cpp
@@ -220,12 +220,7 @@ void ULensComponent::TickComponent(float DeltaTime, ELevelTick TickType, FActorC
 				if (bOverrideCameraOverscan)
 				{
 					const float DesiredOverscan = FMath::Clamp(LensDistortionHandler->GetOverscanFactor() - 1.0f, 0.0f, 1.0f);
-
-					// Only override the camera overscan if the existing amount is too small for lens distortion
-					if (CineCameraComponent->Overscan < DesiredOverscan)
-					{
-						CineCameraComponent->Overscan = DesiredOverscan;
-					}
+					CineCameraComponent->Overscan = DesiredOverscan;
 
 					// The overscan factor used by the post process material needs to be equivalent to what is set on the camera in order for the PPM and SVE paths to render the same
 					LensDistortionHandler->SetOverscanFactor(CineCameraComponent->Overscan + 1.0f);
