Lens Calibration: Added support in the lens distortion SVE for anamorphic lens distortion.

diff --git a/Engine/Plugins/VirtualProduction/CameraCalibrationCore/Shaders/Private/DrawDisplacementMaps.usf b/Engine/Plugins/VirtualProduction/CameraCalibrationCore/Shaders/Private/DrawDisplacementMaps.usf
--- a/Engine/Plugins/VirtualProduction/CameraCalibrationCore/Shaders/Private/DrawDisplacementMaps.usf
+++ b/Engine/Plugins/VirtualProduction/CameraCalibrationCore/Shaders/Private/DrawDisplacementMaps.usf
@@ -5,9 +5,20 @@
 =============================================================================*/
 
 #include "/Engine/Public/Platform.ush"
+#include "/Engine/Private/Common.ush"
+
+#define SPHERICAL_DISTORTION 1
+#define ANAMORPHIC_DISTORTION 2
 
 float2 ThreadIdToUV;
 
+float InverseOverscan;
+float CameraOverscan;
+
+RWTexture2D<float2> OutDistortionMap;
+
+#if DISTORTION_MODEL == SPHERICAL_DISTORTION
+
 float2 FocalLength;
 float2 ImageCenter;
 float K1;
@@ -16,11 +27,6 @@ float K3;
 float P1;
 float P2;
 
-float InverseOverscan;
-float CameraOverscan;
-
-RWTexture2D<float2> OutDistortionMap;
-
 float2 SphericalDistort(float2 InputUV)
 {
 	float x_prime = (InputUV.x - 0.5) / FocalLength.x;
@@ -54,10 +60,90 @@ float2 SphericalDistort(float2 InputUV)
 	return result;
 }
 
+#elif DISTORTION_MODEL == ANAMORPHIC_DISTORTION
+
+float PixelAspect;
+float2 FilmbackSize;
+float2 Squeeze;
+float LensRotation;
+
+float CX02;
+float CX04;
+float CX22;
+float CX24;
+float CX44;
+
+float CY02;
+float CY04;
+float CY22;
+float CY24;
+float CY44;
+
+float2 AnamorphicDistort(float2 InputUV)
+{
+	// The following implementation is based on the 3DE4 Anamorphic - Standard, Degree 4 model for distortion
+	// Pre-compute some values that will be useful later
+	float r_fb = (0.5) * sqrt(FilmbackSize.x*FilmbackSize.x + FilmbackSize.y*FilmbackSize.y);
+
+	float phi_mnt_rad = LensRotation * (PI / 180);
+	float cos_phi_mnt = cos(phi_mnt_rad);
+	float sin_phi_mnt = sin(phi_mnt_rad);
+
+	// Step 1: Transform input UVs into diagonally-normalized coordinates, based on the width and height of the filmback/image
+	float x_dn = (InputUV.x - 0.5) * (FilmbackSize.x / r_fb); // don't forget to factor in lens center offset later
+	float y_dn = (InputUV.y - 0.5) * (FilmbackSize.y / r_fb); // don't forget to factor in lens center offset later
+
+	// Step 2: Scale the x-coordinate by the inverse of the pixel aspect ratio
+	x_dn = x_dn / PixelAspect;
+
+	// Step 3: Rotate x and y by the inverse of the mounting angle
+	float x_dn_rotated = (x_dn * cos_phi_mnt) + (y_dn * sin_phi_mnt);
+	float y_dn_rotated = (y_dn * cos_phi_mnt) - (x_dn * sin_phi_mnt);
+
+	// Step 4: Compute the anamorphic distortion for x and y, to the 4th degree
+	float r_dn_sq = x_dn_rotated*x_dn_rotated + y_dn_rotated*y_dn_rotated;
+	float r_dn_4th = r_dn_sq * r_dn_sq;
+	float phi = atan2(y_dn_rotated, x_dn_rotated);
+
+	float x_dn_prime = x_dn_rotated * (1 + (CX02 * r_dn_sq) + (CX04 * r_dn_4th) + (CX22 * r_dn_sq * cos(2 * phi)) + (CX24 * r_dn_4th * cos(2 * phi)) + (CX44 * r_dn_4th * cos(4 * phi)));
+	float y_dn_prime = y_dn_rotated * (1 + (CY02 * r_dn_sq) + (CY04 * r_dn_4th) + (CY22 * r_dn_sq * cos(2 * phi)) + (CY24 * r_dn_4th * cos(2 * phi)) + (CY44 * r_dn_4th * cos(4 * phi)));
+
+	// Step 5: Scale the x-coordinate by the pixel aspect ratio, and the x- and y-coordinates by the lens-breathing squeeze factors
+	x_dn_prime = x_dn_prime * PixelAspect * Squeeze.x;
+	y_dn_prime = y_dn_prime * Squeeze.y;
+
+	// Step 6: Rotate x and y by the mounting angle
+	float x_dn_prime_rotated = (x_dn_prime * cos_phi_mnt) - (y_dn_prime * sin_phi_mnt);
+	float y_dn_prime_rotated = (y_dn_prime * cos_phi_mnt) + (x_dn_prime * sin_phi_mnt);
+
+	// Step 7: Transform diagonally normalized coordinates back into unit coordinates
+	float x_unit_prime = ((x_dn_prime_rotated * r_fb) / FilmbackSize.x) + 0.5;
+	float y_unit_prime = ((y_dn_prime_rotated * r_fb) / FilmbackSize.y) + 0.5;
+	
+	float2 output_uv;
+	output_uv.x = x_unit_prime;
+	output_uv.y = y_unit_prime;
+
+	return output_uv;
+}
+
+#endif
+
+float2 Distort(float2 InputUV)
+{
+#if DISTORTION_MODEL == SPHERICAL_DISTORTION
+	return SphericalDistort(InputUV);
+#elif DISTORTION_MODEL == ANAMORPHIC_DISTORTION
+	return AnamorphicDistort(InputUV);
+#else
+	return InputUV;
+#endif
+}
+
 bool IsDistortionWellOriented(float2 InputUV, float2 DistortedUV, float2 UVOffset)
 {
-	float2 DistortedXOffset = SphericalDistort(InputUV + UVOffset * int2(1, 0));
-	float2 DistortedYOffset = SphericalDistort(InputUV + UVOffset * int2(0, 1));
+	float2 DistortedXOffset = Distort(InputUV + UVOffset * int2(1, 0));
+	float2 DistortedYOffset = Distort(InputUV + UVOffset * int2(0, 1));
 
 	float2 XOffset = DistortedXOffset - DistortedUV;
 	float2 YOffset = DistortedYOffset - DistortedUV;
@@ -67,9 +153,9 @@ bool IsDistortionWellOriented(float2 InputUV, float2 DistortedUV, float2 UVOffse
 	return bIsFrontFacing && bIsRightSideUp;
 }
 
-float2 TrySphericalDistort(float2 InputUV)
+float2 TryDistort(float2 InputUV)
 {
-	float2 DistortedPosition = SphericalDistort(InputUV);
+	float2 DistortedPosition = Distort(InputUV);
 	float2 UVOffset = 0.5 * ThreadIdToUV;
 	bool bIsOrientedWell =
 		IsDistortionWellOriented(InputUV, DistortedPosition, UVOffset) &&
@@ -87,7 +173,7 @@ void MainCS(uint3 DispatchThreadId : SV_DispatchThreadID)
 	// Rescale the UV by an overscan factor (>= 1.0) to widen the range of UVs to be distorted
 	UndistortedUV = ((UndistortedUV - 0.5) * InverseOverscan * CameraOverscan) + 0.5;
 	
-	float2 DistortedUV = TrySphericalDistort(UndistortedUV);
+	float2 DistortedUV = TryDistort(UndistortedUV);
 	
 	// Compute the UV displacement and scale it by the CameraOverscan to account for the increased resolution of the overscanned scene texture 
 	OutDistortionMap[DispatchThreadId.xy] = (DistortedUV - UndistortedUV) / CameraOverscan;
diff --git a/Engine/Plugins/VirtualProduction/CameraCalibrationCore/Source/CameraCalibrationCore/Private/AnamorphicLensDistortionModelHandler.cpp b/Engine/Plugins/VirtualProduction/CameraCalibrationCore/Source/CameraCalibrationCore/Private/AnamorphicLensDistortionModelHandler.cpp
--- a/Engine/Plugins/VirtualProduction/CameraCalibrationCore/Source/CameraCalibrationCore/Private/AnamorphicLensDistortionModelHandler.cpp
+++ b/Engine/Plugins/VirtualProduction/CameraCalibrationCore/Source/CameraCalibrationCore/Private/AnamorphicLensDistortionModelHandler.cpp
@@ -10,7 +10,65 @@ void UAnamorphicLensDistortionModelHandler::InitializeHandler()
 	LensModelClass = UAnamorphicLensModel::StaticClass();
 }
 
-FVector2D UAnamorphicLensDistortionModelHandler::ComputeDistortedUV(const FVector2D& InUndistortedUV) const
+FVector2D UAnamorphicLensDistortionModelHandler::ComputeDistortedUV(const FVector2D& InScreenUV) const
+{
+	// Pre-compute some values that will be useful later
+	const FVector2D FilmbackSize = FVector2D(CameraFilmback.SensorWidth * AnamorphicParameters.PixelAspect, CameraFilmback.SensorHeight);
+
+	// Guard against divide-by-zero errors
+	if (FMath::IsNearlyZero(FilmbackSize.X) || FMath::IsNearlyZero(FilmbackSize.Y) || FMath::IsNearlyZero(AnamorphicParameters.PixelAspect))
+	{
+		return InScreenUV;
+	}
+	
+	// Pre-compute some values that will be useful later
+	const double FilmbackRadius = (0.5) * sqrt(FilmbackSize.X * FilmbackSize.X + FilmbackSize.Y * FilmbackSize.Y);
+
+	const double MountAngleRadians = FMath::DegreesToRadians(AnamorphicParameters.LensRotation);
+	const double CosMountAngle = FMath::Cos(MountAngleRadians);
+	const double SinMountAngle = FMath::Sin(MountAngleRadians);
+
+	// Step 1: Transform input UVs into diagonally-normalized coordinates, based on the width and height of the filmback/image
+	FVector2D DNCoordinates = (InScreenUV - 0.5f) * (FilmbackSize / FilmbackRadius);
+
+	// Step 2: Rotate x and y by the inverse of the mounting angle
+	FVector2D DNRotatedCoordinates;
+	DNRotatedCoordinates.X = (DNCoordinates.X * CosMountAngle) + (DNCoordinates.Y * SinMountAngle);
+	DNRotatedCoordinates.Y = (DNCoordinates.Y * CosMountAngle) - (DNCoordinates.X * SinMountAngle);
+
+	// Step 3: Scale the x-coordinate by the inverse of the pixel aspect ratio
+	DNRotatedCoordinates.X = DNRotatedCoordinates.X / AnamorphicParameters.PixelAspect;
+
+	FVector2D DNCoordinatesDistorted = DNRotatedCoordinates;
+
+	// Step 4: Compute the anamorphic distortion for x and y, to the 4th degree
+	constexpr uint32 NumIterations = 10;
+	for(int i = 0; i < NumIterations; i++)
+	{
+		const double DNRadiusSq = (DNCoordinatesDistorted.X * DNCoordinatesDistorted.X) + (DNCoordinatesDistorted.Y * DNCoordinatesDistorted.Y);
+		const double DNRadius4th = DNRadiusSq * DNRadiusSq;
+		const double Phi = FMath::Atan2(DNCoordinatesDistorted.Y, DNCoordinatesDistorted.X);
+
+		const double DistX = (1 + (AnamorphicParameters.CX02 * DNRadiusSq) + (AnamorphicParameters.CX04 * DNRadius4th) + (AnamorphicParameters.CX22 * DNRadiusSq * FMath::Cos(2 * Phi)) + (AnamorphicParameters.CX24 * DNRadius4th * FMath::Cos(2 * Phi)) + (AnamorphicParameters.CX44 * DNRadius4th * FMath::Cos(4 * Phi)));
+		const double DistY = (1 + (AnamorphicParameters.CY02 * DNRadiusSq) + (AnamorphicParameters.CY04 * DNRadius4th) + (AnamorphicParameters.CY22 * DNRadiusSq * FMath::Cos(2 * Phi)) + (AnamorphicParameters.CY24 * DNRadius4th * FMath::Cos(2 * Phi)) + (AnamorphicParameters.CY44 * DNRadius4th * FMath::Cos(4 * Phi)));
+
+		DNCoordinatesDistorted.X = (DNRotatedCoordinates.X / DistX) / AnamorphicParameters.SqueezeX;
+		DNCoordinatesDistorted.Y = (DNRotatedCoordinates.Y / DistY) / AnamorphicParameters.SqueezeY;
+	}
+
+	// Step 5: Rotate x and y by the mounting angle
+	DNCoordinates.X = (DNCoordinatesDistorted.X * CosMountAngle) - (DNCoordinatesDistorted.Y * SinMountAngle);
+	DNCoordinates.Y = (DNCoordinatesDistorted.Y * CosMountAngle) + (DNCoordinatesDistorted.X * SinMountAngle);
+
+	// Step 6: Scale the x-coordinate by the pixel aspect ratio
+	DNCoordinates.X = DNCoordinates.X * AnamorphicParameters.PixelAspect;
+
+	// Step 7: Transform diagonally normalized coordinates back into unit coordinates
+	const FVector2D DistortedUV = ((DNCoordinates * FilmbackRadius) / FilmbackSize) + 0.5f;
+	return DistortedUV;
+}
+
+FVector2D UAnamorphicLensDistortionModelHandler::ComputeUndistortedUV(const FVector2D& InScreenUV) const
 {
 	// The following implementation is based on the 3DE4 Anamorphic - Standard, Degree 4 model for distortion
 
@@ -20,7 +78,7 @@ FVector2D UAnamorphicLensDistortionModelHandler::ComputeDistortedUV(const FVecto
 	// Guard against divide-by-zero errors
 	if (FMath::IsNearlyZero(FilmbackSize.X) || FMath::IsNearlyZero(FilmbackSize.Y) || FMath::IsNearlyZero(AnamorphicParameters.PixelAspect))
 	{
-		return InUndistortedUV;
+		return InScreenUV;
 	}
 
 	const double FilmbackRadius = 0.5f * FMath::Sqrt((FilmbackSize.X * FilmbackSize.X) + (FilmbackSize.Y * FilmbackSize.Y));
@@ -30,34 +88,35 @@ FVector2D UAnamorphicLensDistortionModelHandler::ComputeDistortedUV(const FVecto
 	const double SinMountAngle = FMath::Sin(MountAngleRadians);
 
 	// Step 1: Transform input UVs into diagonally-normalized coordinates, based on the width and height of the filmback/image
-	FVector2D DNCoordinates = (InUndistortedUV - 0.5f) * (FilmbackSize / FilmbackRadius);
+	FVector2D DNCoordinates = (InScreenUV - 0.5f) * (FilmbackSize / FilmbackRadius);
 
 	// Step 2: Scale the x-coordinate by the inverse of the pixel aspect ratio
 	DNCoordinates.X = DNCoordinates.X / AnamorphicParameters.PixelAspect;
 
 	// Step 3: Rotate x and y by the inverse of the mounting angle
-	DNCoordinates.X = (DNCoordinates.X * CosMountAngle) + (DNCoordinates.Y * SinMountAngle);
-	DNCoordinates.Y = (DNCoordinates.Y * CosMountAngle) - (DNCoordinates.X * SinMountAngle);
+	FVector2D DNRotatedCoordinates;
+	DNRotatedCoordinates.X = (DNCoordinates.X * CosMountAngle) + (DNCoordinates.Y * SinMountAngle);
+	DNRotatedCoordinates.Y = (DNCoordinates.Y * CosMountAngle) - (DNCoordinates.X * SinMountAngle);
 
 	// Step 4: Compute the anamorphic distortion for x and y, to the 4th degree
-	const double DNRadiusSq = (DNCoordinates.X * DNCoordinates.X) + (DNCoordinates.Y * DNCoordinates.Y);
+	const double DNRadiusSq = (DNRotatedCoordinates.X * DNRotatedCoordinates.X) + (DNRotatedCoordinates.Y * DNRotatedCoordinates.Y);
 	const double DNRadius4th = DNRadiusSq * DNRadiusSq;
-	const double Phi = FMath::Atan2(DNCoordinates.Y, DNCoordinates.X);
+	const double Phi = FMath::Atan2(DNRotatedCoordinates.Y, DNRotatedCoordinates.X);
 
 	FVector2D DNCoordinatesDistorted;
-	DNCoordinatesDistorted.X = DNCoordinates.X * (1 + (AnamorphicParameters.CX02 * DNRadiusSq) + (AnamorphicParameters.CX04 * DNRadius4th) + (AnamorphicParameters.CX22 * DNRadiusSq * FMath::Cos(2 * Phi)) + (AnamorphicParameters.CX24 * DNRadius4th * FMath::Cos(2 * Phi)) + (AnamorphicParameters.CX44 * DNRadius4th * FMath::Cos(4 * Phi)));
-	DNCoordinatesDistorted.Y = DNCoordinates.Y * (1 + (AnamorphicParameters.CY02 * DNRadiusSq) + (AnamorphicParameters.CY04 * DNRadius4th) + (AnamorphicParameters.CY22 * DNRadiusSq * FMath::Cos(2 * Phi)) + (AnamorphicParameters.CY24 * DNRadius4th * FMath::Cos(2 * Phi)) + (AnamorphicParameters.CY44 * DNRadius4th * FMath::Cos(4 * Phi)));
+	DNCoordinatesDistorted.X = DNRotatedCoordinates.X * (1 + (AnamorphicParameters.CX02 * DNRadiusSq) + (AnamorphicParameters.CX04 * DNRadius4th) + (AnamorphicParameters.CX22 * DNRadiusSq * FMath::Cos(2 * Phi)) + (AnamorphicParameters.CX24 * DNRadius4th * FMath::Cos(2 * Phi)) + (AnamorphicParameters.CX44 * DNRadius4th * FMath::Cos(4 * Phi)));
+	DNCoordinatesDistorted.Y = DNRotatedCoordinates.Y * (1 + (AnamorphicParameters.CY02 * DNRadiusSq) + (AnamorphicParameters.CY04 * DNRadius4th) + (AnamorphicParameters.CY22 * DNRadiusSq * FMath::Cos(2 * Phi)) + (AnamorphicParameters.CY24 * DNRadius4th * FMath::Cos(2 * Phi)) + (AnamorphicParameters.CY44 * DNRadius4th * FMath::Cos(4 * Phi)));
 
 	// Step 5: Scale the x-coordinate by the pixel aspect ratio, and the x- and y-coordinates by the lens-breathing squeeze factors
 	DNCoordinatesDistorted.X = DNCoordinatesDistorted.X * AnamorphicParameters.PixelAspect * AnamorphicParameters.SqueezeX;
 	DNCoordinatesDistorted.Y = DNCoordinatesDistorted.Y * AnamorphicParameters.SqueezeY;
 
 	// Step 6: Rotate x and y by the mounting angle
-	DNCoordinates.X = (DNCoordinates.X * CosMountAngle) - (DNCoordinates.Y * SinMountAngle);
-	DNCoordinates.Y = (DNCoordinates.Y * CosMountAngle) + (DNCoordinates.X * SinMountAngle);
+	DNCoordinates.X = (DNCoordinatesDistorted.X * CosMountAngle) - (DNCoordinatesDistorted.Y * SinMountAngle);
+	DNCoordinates.Y = (DNCoordinatesDistorted.Y * CosMountAngle) + (DNCoordinatesDistorted.X * SinMountAngle);
 
 	// Step 7: Transform diagonally normalized coordinates back into unit coordinates
-	const FVector2D DistortedUV = ((DNCoordinatesDistorted * FilmbackRadius) / FilmbackSize) + 0.5f;
+	const FVector2D DistortedUV = ((DNCoordinates * FilmbackRadius) / FilmbackSize) + 0.5f;
 
 	return DistortedUV;
 }
diff --git a/Engine/Plugins/VirtualProduction/CameraCalibrationCore/Source/CameraCalibrationCore/Private/CameraCalibrationSubsystem.cpp b/Engine/Plugins/VirtualProduction/CameraCalibrationCore/Source/CameraCalibrationCore/Private/CameraCalibrationSubsystem.cpp
--- a/Engine/Plugins/VirtualProduction/CameraCalibrationCore/Source/CameraCalibrationCore/Private/CameraCalibrationSubsystem.cpp
+++ b/Engine/Plugins/VirtualProduction/CameraCalibrationCore/Source/CameraCalibrationCore/Private/CameraCalibrationSubsystem.cpp
@@ -176,9 +176,9 @@ TArray<FName> UCameraCalibrationSubsystem::GetCameraCalibrationSteps() const
 	return OutKeys;
 }
 
-void UCameraCalibrationSubsystem::SetLensDistortionSVEState(ACameraActor* CameraActor, FDisplacementMapBlendingParams DistortionState)
+void UCameraCalibrationSubsystem::SetLensDistortionSVEState(ACameraActor* CameraActor, FDisplacementMapBlendingParams DistortionState, ULensDistortionModelHandlerBase* LensDistortionHandler)
 {
-	SceneViewExtension->UpdateDistortionState_AnyThread(CameraActor, DistortionState);
+	SceneViewExtension->UpdateDistortionState_AnyThread(CameraActor, DistortionState, LensDistortionHandler);
 }
 
 void UCameraCalibrationSubsystem::ClearLensDistortionSVEState(ACameraActor* CameraActor)
diff --git a/Engine/Plugins/VirtualProduction/CameraCalibrationCore/Source/CameraCalibrationCore/Private/LensDistortionModelHandlerBase.cpp b/Engine/Plugins/VirtualProduction/CameraCalibrationCore/Source/CameraCalibrationCore/Private/LensDistortionModelHandlerBase.cpp
--- a/Engine/Plugins/VirtualProduction/CameraCalibrationCore/Source/CameraCalibrationCore/Private/LensDistortionModelHandlerBase.cpp
+++ b/Engine/Plugins/VirtualProduction/CameraCalibrationCore/Source/CameraCalibrationCore/Private/LensDistortionModelHandlerBase.cpp
@@ -140,6 +140,49 @@ void ULensDistortionModelHandlerBase::SetOverscanFactor(float InOverscanFactor)
 }
 
 float ULensDistortionModelHandlerBase::ComputeOverscanFactor() const
+{
+	/* Distorted UV position in the view space:
+					   ^ View space's Y
+					   |
+			  0        1        2
+
+			  7                 3 --> View space's X
+
+			  6        5        4
+	*/
+	const TArray<FVector2D> DistortedUVs =
+	{
+		FVector2D(0.0f, 0.0f),
+		FVector2D(0.5f, 0.0f),
+		FVector2D(1.0f, 0.0f),
+		FVector2D(1.0f, 0.5f),
+		FVector2D(1.0f, 1.0f),
+		FVector2D(0.5f, 1.0f),
+		FVector2D(0.0f, 1.0f),
+		FVector2D(0.0f, 0.5f)
+	};
+
+	TArray<float> OverscanFactors;
+	OverscanFactors.Reserve(DistortedUVs.Num());
+	for (const FVector2D& DistortedUV : DistortedUVs)
+	{
+		const FVector2D UndistortedUV = ComputeUndistortedUV(DistortedUV);
+		const float OverscanX = (DistortedUV.X != 0.5f) ? (UndistortedUV.X - 0.5f) / (DistortedUV.X - 0.5f) : 1.0f;
+		const float OverscanY = (DistortedUV.Y != 0.5f) ? (UndistortedUV.Y - 0.5f) / (DistortedUV.Y - 0.5f) : 1.0f;
+		OverscanFactors.Add(FMath::Max(OverscanX, OverscanY));
+	}
+
+	float NewOverscan = 1.0f;
+	float* MaxOverscanFactor = Algo::MaxElement(OverscanFactors);
+	if (MaxOverscanFactor != nullptr)
+	{
+		NewOverscan = FMath::Max(*MaxOverscanFactor, 1.0f);
+	}
+
+	return NewOverscan;
+}
+
+float ULensDistortionModelHandlerBase::ComputeInverseOverscanFactor() const
 {
 	/* Undistorted UV position in the view space:
 					   ^ View space's Y
diff --git a/Engine/Plugins/VirtualProduction/CameraCalibrationCore/Source/CameraCalibrationCore/Private/LensDistortionSceneViewExtension.cpp b/Engine/Plugins/VirtualProduction/CameraCalibrationCore/Source/CameraCalibrationCore/Private/LensDistortionSceneViewExtension.cpp
--- a/Engine/Plugins/VirtualProduction/CameraCalibrationCore/Source/CameraCalibrationCore/Private/LensDistortionSceneViewExtension.cpp
+++ b/Engine/Plugins/VirtualProduction/CameraCalibrationCore/Source/CameraCalibrationCore/Private/LensDistortionSceneViewExtension.cpp
@@ -2,6 +2,7 @@
 
 #include "LensDistortionSceneViewExtension.h"
 
+#include "CineCameraComponent.h"
 #include "Camera/CameraActor.h"
 #include "Camera/CameraComponent.h"
 #include "Containers/DynamicRHIResourceArray.h"
@@ -11,6 +12,8 @@
 #include "ScreenPass.h"
 #include "ShaderParameterStruct.h"
 #include "SystemTextures.h"
+#include "Models/AnamorphicLensModel.h"
+#include "Models/SphericalLensModel.h"
 
 TAutoConsoleVariable<int32> CVarLensDistortionInvertGridDensity(
 	TEXT("r.LensDistortion.InvertGridDensity"),
@@ -24,10 +27,10 @@ FLensDistortionSceneViewExtension::FLensDistortionSceneViewExtension(const FAuto
 {
 }
 
-void FLensDistortionSceneViewExtension::UpdateDistortionState_AnyThread(ACameraActor* CameraActor, FDisplacementMapBlendingParams DistortionState)
+void FLensDistortionSceneViewExtension::UpdateDistortionState_AnyThread(ACameraActor* CameraActor, FDisplacementMapBlendingParams DistortionState, ULensDistortionModelHandlerBase* LensDistortionHandler)
 {
 	FScopeLock ScopeLock(&DistortionStateMapCriticalSection);
-	DistortionStateMap.Add(CameraActor->GetUniqueID(), {DistortionState, CameraActor->GetCameraComponent()});
+	DistortionStateMap.Add(CameraActor->GetUniqueID(), {DistortionState, CameraActor->GetCameraComponent(), LensDistortionHandler});
 }
 
 void FLensDistortionSceneViewExtension::ClearDistortionState_AnyThread(ACameraActor* CameraActor)
@@ -36,22 +39,47 @@ void FLensDistortionSceneViewExtension::ClearDistortionState_AnyThread(ACameraAc
 	DistortionStateMap.Remove(CameraActor->GetUniqueID());
 }
 
+BEGIN_SHADER_PARAMETER_STRUCT(FSphericalDistortionParams, )
+	SHADER_PARAMETER(FVector2f, FocalLength)
+	SHADER_PARAMETER(FVector2f, ImageCenter)
+	SHADER_PARAMETER(float, K1)
+	SHADER_PARAMETER(float, K2)
+	SHADER_PARAMETER(float, K3)
+	SHADER_PARAMETER(float, P1)
+	SHADER_PARAMETER(float, P2)
+END_SHADER_PARAMETER_STRUCT()
+
+BEGIN_SHADER_PARAMETER_STRUCT(FAnamorphicDistortionParams, )
+	SHADER_PARAMETER(float, PixelAspect)
+	SHADER_PARAMETER(FVector2f, FilmbackSize)
+	SHADER_PARAMETER(FVector2f, Squeeze)
+	SHADER_PARAMETER(float, LensRotation)
+	SHADER_PARAMETER(float, CX02)
+	SHADER_PARAMETER(float, CX04)
+	SHADER_PARAMETER(float, CX22)
+	SHADER_PARAMETER(float, CX24)
+	SHADER_PARAMETER(float, CX44)
+	SHADER_PARAMETER(float, CY02)
+	SHADER_PARAMETER(float, CY04)
+	SHADER_PARAMETER(float, CY22)
+	SHADER_PARAMETER(float, CY24)
+	SHADER_PARAMETER(float, CY44)
+END_SHADER_PARAMETER_STRUCT()
+
 class FDrawDistortionDisplacementMapCS : public FGlobalShader
 {
 	DECLARE_GLOBAL_SHADER(FDrawDistortionDisplacementMapCS);
 	SHADER_USE_PARAMETER_STRUCT(FDrawDistortionDisplacementMapCS, FGlobalShader);
 
+	class FDistortionModelDim : SHADER_PERMUTATION_ENUM_CLASS("DISTORTION_MODEL", FLensDistortionSceneViewExtension::EDistortionModel);
+	using FPermutationDomain = TShaderPermutationDomain<FDistortionModelDim>;
+	
 	BEGIN_SHADER_PARAMETER_STRUCT(FParameters, )
 		SHADER_PARAMETER(FVector2f, ThreadIdToUV)
-		SHADER_PARAMETER(FVector2f, FocalLength)
-		SHADER_PARAMETER(FVector2f, ImageCenter)
-		SHADER_PARAMETER(float, K1)
-		SHADER_PARAMETER(float, K2)
-		SHADER_PARAMETER(float, K3)
-		SHADER_PARAMETER(float, P1)
-		SHADER_PARAMETER(float, P2)
 		SHADER_PARAMETER(float, InverseOverscan)
 		SHADER_PARAMETER(float, CameraOverscan)
+		SHADER_PARAMETER_STRUCT_INCLUDE(FSphericalDistortionParams, SphericalDistortionParams)
+		SHADER_PARAMETER_STRUCT_INCLUDE(FAnamorphicDistortionParams, AnamorphicDistortionParams)
 		SHADER_PARAMETER_RDG_TEXTURE_UAV(RWTexture2D, OutDistortionMap)
 	END_SHADER_PARAMETER_STRUCT()
 };
@@ -133,7 +161,7 @@ public:
 IMPLEMENT_GLOBAL_SHADER(FInvertDisplacementVS, "/Plugin/CameraCalibrationCore/Private/InvertDisplacementMap.usf", "MainVS", SF_Vertex);
 IMPLEMENT_GLOBAL_SHADER(FInvertDisplacementPS, "/Plugin/CameraCalibrationCore/Private/InvertDisplacementMap.usf", "MainPS", SF_Pixel);
 
-void FLensDistortionSceneViewExtension::DrawDisplacementMap_RenderThread(FRDGBuilder& GraphBuilder, const FLensDistortionState& CurrentState, float InverseOverscan, float CameraOverscan, FRDGTextureRef& OutDistortionMapWithOverscan)
+void FLensDistortionSceneViewExtension::DrawDisplacementMap_RenderThread(FRDGBuilder& GraphBuilder, const FLensDistortionState& CurrentState, EDistortionModel DistortionModel, float InverseOverscan, float CameraOverscan, const FVector2D& SensorSize, FRDGTextureRef& OutDistortionMapWithOverscan)
 {
 	if (CurrentState.DistortionInfo.Parameters.IsEmpty())
 	{
@@ -141,6 +169,9 @@ void FLensDistortionSceneViewExtension::DrawDisplacementMap_RenderThread(FRDGBui
 		return;
 	}
 
+	FDrawDistortionDisplacementMapCS::FPermutationDomain PermutationVector;
+	PermutationVector.Set<FDrawDistortionDisplacementMapCS::FDistortionModelDim>(DistortionModel);
+	
 	FDrawDistortionDisplacementMapCS::FParameters* PassParameters = GraphBuilder.AllocParameters<FDrawDistortionDisplacementMapCS::FParameters>();
 
 	PassParameters->OutDistortionMap = GraphBuilder.CreateUAV(OutDistortionMapWithOverscan);
@@ -148,19 +179,39 @@ void FLensDistortionSceneViewExtension::DrawDisplacementMap_RenderThread(FRDGBui
 	FIntPoint DistortionMapResolution = OutDistortionMapWithOverscan->Desc.Extent;
 	PassParameters->ThreadIdToUV = FVector2f(1.0f) / FVector2f(DistortionMapResolution);
 
-	PassParameters->ImageCenter = FVector2f(CurrentState.ImageCenter.PrincipalPoint);
-	PassParameters->FocalLength = FVector2f(CurrentState.FocalLengthInfo.FxFy);
-
-	PassParameters->K1 = CurrentState.DistortionInfo.Parameters[0];
-	PassParameters->K2 = CurrentState.DistortionInfo.Parameters[1];
-	PassParameters->K3 = CurrentState.DistortionInfo.Parameters[2];
-	PassParameters->P1 = CurrentState.DistortionInfo.Parameters[3];
-	PassParameters->P2 = CurrentState.DistortionInfo.Parameters[4];
-
+	if (DistortionModel == EDistortionModel::SphericalDistortion)
+	{
+		PassParameters->SphericalDistortionParams.ImageCenter = FVector2f(CurrentState.ImageCenter.PrincipalPoint);
+		PassParameters->SphericalDistortionParams.FocalLength = FVector2f(CurrentState.FocalLengthInfo.FxFy);
+
+		PassParameters->SphericalDistortionParams.K1 = CurrentState.DistortionInfo.Parameters[0];
+		PassParameters->SphericalDistortionParams.K2 = CurrentState.DistortionInfo.Parameters[1];
+		PassParameters->SphericalDistortionParams.K3 = CurrentState.DistortionInfo.Parameters[2];
+		PassParameters->SphericalDistortionParams.P1 = CurrentState.DistortionInfo.Parameters[3];
+		PassParameters->SphericalDistortionParams.P2 = CurrentState.DistortionInfo.Parameters[4];
+	}
+	else if (DistortionModel == EDistortionModel::AnamorphicDistortion)
+	{
+		PassParameters->AnamorphicDistortionParams.FilmbackSize = FVector2f(SensorSize.X * CurrentState.DistortionInfo.Parameters[0], SensorSize.Y);
+		PassParameters->AnamorphicDistortionParams.PixelAspect = CurrentState.DistortionInfo.Parameters[0];
+		PassParameters->AnamorphicDistortionParams.CX02 = CurrentState.DistortionInfo.Parameters[1];
+		PassParameters->AnamorphicDistortionParams.CX04 = CurrentState.DistortionInfo.Parameters[2];
+		PassParameters->AnamorphicDistortionParams.CX22 = CurrentState.DistortionInfo.Parameters[3];
+		PassParameters->AnamorphicDistortionParams.CX24 = CurrentState.DistortionInfo.Parameters[4];
+		PassParameters->AnamorphicDistortionParams.CX44 = CurrentState.DistortionInfo.Parameters[5];
+		PassParameters->AnamorphicDistortionParams.CY02 = CurrentState.DistortionInfo.Parameters[6];
+		PassParameters->AnamorphicDistortionParams.CY04 = CurrentState.DistortionInfo.Parameters[7];
+		PassParameters->AnamorphicDistortionParams.CY22 = CurrentState.DistortionInfo.Parameters[8];
+		PassParameters->AnamorphicDistortionParams.CY24 = CurrentState.DistortionInfo.Parameters[9];
+		PassParameters->AnamorphicDistortionParams.CY44 = CurrentState.DistortionInfo.Parameters[10];
+		PassParameters->AnamorphicDistortionParams.Squeeze = FVector2f(CurrentState.DistortionInfo.Parameters[11], CurrentState.DistortionInfo.Parameters[12]);
+		PassParameters->AnamorphicDistortionParams.LensRotation = CurrentState.DistortionInfo.Parameters[13];
+	}
+	
 	PassParameters->InverseOverscan = InverseOverscan;
 	PassParameters->CameraOverscan = CameraOverscan;
 
-	TShaderMapRef<FDrawDistortionDisplacementMapCS> ComputeShader(GetGlobalShaderMap(GMaxRHIFeatureLevel));
+	TShaderMapRef<FDrawDistortionDisplacementMapCS> ComputeShader(GetGlobalShaderMap(GMaxRHIFeatureLevel), PermutationVector);
 	FComputeShaderUtils::AddPass(
 		GraphBuilder,
 		RDG_EVENT_NAME("DrawDistortionDisplacementMap"),
@@ -188,14 +239,14 @@ void FLensDistortionSceneViewExtension::CropDisplacementMap_RenderThread(FRDGBui
 		FIntVector(FMath::DivideAndRoundUp(LUTResolution.X, 8), FMath::DivideAndRoundUp(LUTResolution.Y, 8), 1));
 }
 
-void FLensDistortionSceneViewExtension::BlendDisplacementMaps_RenderThread(FRDGBuilder& GraphBuilder, const FDisplacementMapBlendingParams& BlendState, float InverseOverscan, float CameraOverscan, FRDGTextureRef& OutDistortionMapWithOverscan)
+void FLensDistortionSceneViewExtension::BlendDisplacementMaps_RenderThread(FRDGBuilder& GraphBuilder, const FDisplacementMapBlendingParams& BlendState, EDistortionModel DistortionModel, float InverseOverscan, float CameraOverscan, const FVector2D& SensorSize, FRDGTextureRef& OutDistortionMapWithOverscan)
 {
 	FBlendDistortionDisplacementMapCS::FParameters* PassParameters = GraphBuilder.AllocParameters<FBlendDistortionDisplacementMapCS::FParameters>();
 
 	// Draw the first distortion map, which should always be valid
 	{
 		FRDGTextureRef Distortion1 = GraphBuilder.CreateTexture(OutDistortionMapWithOverscan->Desc, TEXT("DistortingDisplacement1"));
-		DrawDisplacementMap_RenderThread(GraphBuilder, BlendState.States[0], InverseOverscan, CameraOverscan, Distortion1);
+		DrawDisplacementMap_RenderThread(GraphBuilder, BlendState.States[0], DistortionModel, InverseOverscan, CameraOverscan, SensorSize, Distortion1);
 		PassParameters->InputDistortionMap1 = GraphBuilder.CreateSRV(Distortion1);
 	}
 
@@ -203,7 +254,7 @@ void FLensDistortionSceneViewExtension::BlendDisplacementMaps_RenderThread(FRDGB
 	if (BlendState.BlendType != EDisplacementMapBlendType::OneFocusOneZoom)
 	{
 		FRDGTextureRef Distortion2 = GraphBuilder.CreateTexture(OutDistortionMapWithOverscan->Desc, TEXT("DistortingDisplacement2"));
-		DrawDisplacementMap_RenderThread(GraphBuilder, BlendState.States[1], InverseOverscan, CameraOverscan, Distortion2);
+		DrawDisplacementMap_RenderThread(GraphBuilder, BlendState.States[1], DistortionModel, InverseOverscan, CameraOverscan, SensorSize, Distortion2);
 		PassParameters->InputDistortionMap2 = GraphBuilder.CreateSRV(Distortion2);
 	}
 
@@ -213,8 +264,8 @@ void FLensDistortionSceneViewExtension::BlendDisplacementMaps_RenderThread(FRDGB
 		FRDGTextureRef Distortion3 = GraphBuilder.CreateTexture(OutDistortionMapWithOverscan->Desc, TEXT("DistortingDisplacement3"));
 		FRDGTextureRef Distortion4 = GraphBuilder.CreateTexture(OutDistortionMapWithOverscan->Desc, TEXT("DistortingDisplacement4"));
 
-		DrawDisplacementMap_RenderThread(GraphBuilder, BlendState.States[2], InverseOverscan, CameraOverscan, Distortion3);
-		DrawDisplacementMap_RenderThread(GraphBuilder, BlendState.States[3], InverseOverscan, CameraOverscan, Distortion4);
+		DrawDisplacementMap_RenderThread(GraphBuilder, BlendState.States[2], DistortionModel, InverseOverscan, CameraOverscan, SensorSize, Distortion3);
+		DrawDisplacementMap_RenderThread(GraphBuilder, BlendState.States[3], DistortionModel, InverseOverscan, CameraOverscan, SensorSize, Distortion4);
 
 		PassParameters->InputDistortionMap3 = GraphBuilder.CreateSRV(Distortion3);
 		PassParameters->InputDistortionMap4 = GraphBuilder.CreateSRV(Distortion4);
@@ -360,6 +411,31 @@ void FLensDistortionSceneViewExtension::InvertDistortionMap_RenderThread(FRDGBui
 		});
 }
 
+bool FLensDistortionSceneViewExtension::IsDistortionModelForwardDistorting(EDistortionModel InDistortionModel) const
+{
+	switch (InDistortionModel)
+	{
+	case EDistortionModel::AnamorphicDistortion:
+		return false;
+
+	case EDistortionModel::SphericalDistortion:
+	default:
+		return true;
+	}
+}
+
+float FLensDistortionSceneViewExtension::GetInverseOverscan(ULensDistortionModelHandlerBase* LensDistortionHandler, EDistortionModel InDistortionModel) const
+{
+	if (IsDistortionModelForwardDistorting(InDistortionModel))
+	{
+		return LensDistortionHandler->ComputeOverscanFactor();
+	}
+	else
+	{
+		return LensDistortionHandler->ComputeInverseOverscanFactor();
+	}
+}
+
 void FLensDistortionSceneViewExtension::PreRenderView_RenderThread(FRDGBuilder& GraphBuilder, FSceneView& InView)
 {
 	FScopeLock ScopeLock(&DistortionStateMapCriticalSection);
@@ -369,22 +445,44 @@ void FLensDistortionSceneViewExtension::PreRenderView_RenderThread(FRDGBuilder&
 		return;
 	}
 
+	FCameraFilmbackSettings CameraFilmback;
+	if (UCineCameraComponent* CineCameraComponent = Cast<UCineCameraComponent>(Rec->Component))
+	{
+		CameraFilmback = CineCameraComponent->Filmback;
+	}
+
+	FVector2D SensorSize = FVector2D(CameraFilmback.SensorWidth, CameraFilmback.SensorHeight);
 	float CameraOverscan = Rec->Component ? (Rec->Component->Overscan + 1.0f) : 1.0f;
 	FDisplacementMapBlendingParams BlendState = Rec->Params;
+	
+	EDistortionModel DistortionModel = EDistortionModel::None;
+	if (Rec->LensDistortionHandler)
+	{
+		if (Rec->LensDistortionHandler->GetLensModelClass() == USphericalLensModel::StaticClass())
+		{
+			DistortionModel = EDistortionModel::SphericalDistortion;
+		}
+		else if (Rec->LensDistortionHandler->GetLensModelClass() == UAnamorphicLensModel::StaticClass())
+		{
+			DistortionModel = EDistortionModel::AnamorphicDistortion;
+		}
+	}
 
+	const bool bForwardDistort = IsDistortionModelForwardDistorting(DistortionModel);
+	
 	ScopeLock.Unlock();
 
 
 	// Create the distortion map and undistortion map textures for the FLensDistortionLUT for this frame
 	const FIntPoint DisplacementMapResolution = FIntPoint(256, 256);
 
-	FRDGTextureDesc DistortionMapDesc = FRDGTextureDesc::Create2D(
+	FRDGTextureDesc ForwardDistortionMapDesc = FRDGTextureDesc::Create2D(
 		DisplacementMapResolution,
 		PF_G32R32F,
 		FClearValueBinding::None,
 		TexCreate_ShaderResource | TexCreate_UAV);
 
-	FRDGTextureDesc UndistortionMapDesc = FRDGTextureDesc::Create2D(
+	FRDGTextureDesc InverseDistortionMapDesc = FRDGTextureDesc::Create2D(
 		DisplacementMapResolution,
 		PF_G32R32F,
 		FClearValueBinding::Black,
@@ -392,13 +490,18 @@ void FLensDistortionSceneViewExtension::PreRenderView_RenderThread(FRDGBuilder&
 
 	FLensDistortionLUT ViewDistortionLUT;
 
+	FRDGTextureDesc DistortionMapDesc = bForwardDistort ? ForwardDistortionMapDesc : InverseDistortionMapDesc;
+	FRDGTextureDesc UndistortionMapDesc = bForwardDistort ? InverseDistortionMapDesc : ForwardDistortionMapDesc;
+	
 	ViewDistortionLUT.DistortingDisplacementTexture = GraphBuilder.CreateTexture(DistortionMapDesc, TEXT("DistortionDisplacementMap"));
 	ViewDistortionLUT.UndistortingDisplacementTexture = GraphBuilder.CreateTexture(UndistortionMapDesc, TEXT("UndistortionDisplacementMap"));
 
 	// In order to guarantee that we can generate a complete undistortion map, the distortion map we invert needs to have some overscan 
+	float InverseOverscan = FMath::Clamp(GetInverseOverscan(Rec->LensDistortionHandler, DistortionModel), 0.0f, 2.0f);
+	
 	// Adjust the overscan resolution to be square, with each side being a multiple of 8
-	const FIntPoint OverscanResolution = FIntPoint(FMath::CeilToInt(CameraOverscan * 32) * 8);
-	const float InverseOverscan = float(OverscanResolution.X) / float(DisplacementMapResolution.X);
+	const FIntPoint OverscanResolution = FIntPoint(FMath::CeilToInt(InverseOverscan * 32) * 8);
+	InverseOverscan = float(OverscanResolution.X) / float(DisplacementMapResolution.X);
 
 	// Create the texture for the overscanned distortion map
 	FRDGTextureDesc OverscanDesc = FRDGTextureDesc::Create2D(
@@ -409,9 +512,13 @@ void FLensDistortionSceneViewExtension::PreRenderView_RenderThread(FRDGBuilder&
 
 	FRDGTextureRef DistortionMapWithOverscan = GraphBuilder.CreateTexture(OverscanDesc, TEXT("DistortionMapWithOverscan"));
 
-	BlendDisplacementMaps_RenderThread(GraphBuilder, BlendState, InverseOverscan, CameraOverscan, DistortionMapWithOverscan);
-	InvertDistortionMap_RenderThread(GraphBuilder, DistortionMapWithOverscan, InverseOverscan, ViewDistortionLUT.UndistortingDisplacementTexture);
-	CropDisplacementMap_RenderThread(GraphBuilder, DistortionMapWithOverscan, ViewDistortionLUT.DistortingDisplacementTexture);
+	// Anamorphic model uses a precise equation for its _undistort_ map instead of its distort map, so draw the undistort map first and invert it to find the distort map
+	FRDGTextureRef ForwardDistortionMap = bForwardDistort ? ViewDistortionLUT.DistortingDisplacementTexture : ViewDistortionLUT.UndistortingDisplacementTexture;
+	FRDGTextureRef InverseDistortionMap = bForwardDistort ? ViewDistortionLUT.UndistortingDisplacementTexture : ViewDistortionLUT.DistortingDisplacementTexture;
+	
+	BlendDisplacementMaps_RenderThread(GraphBuilder, BlendState, DistortionModel, InverseOverscan, CameraOverscan, SensorSize, DistortionMapWithOverscan);
+	InvertDistortionMap_RenderThread(GraphBuilder, DistortionMapWithOverscan, InverseOverscan, InverseDistortionMap);
+	CropDisplacementMap_RenderThread(GraphBuilder, DistortionMapWithOverscan, ForwardDistortionMap);
 
 	LensDistortion::SetLUTUnsafe(InView, ViewDistortionLUT);
 }
diff --git a/Engine/Plugins/VirtualProduction/CameraCalibrationCore/Source/CameraCalibrationCore/Private/SphericalLensDistortionModelHandler.cpp b/Engine/Plugins/VirtualProduction/CameraCalibrationCore/Source/CameraCalibrationCore/Private/SphericalLensDistortionModelHandler.cpp
--- a/Engine/Plugins/VirtualProduction/CameraCalibrationCore/Source/CameraCalibrationCore/Private/SphericalLensDistortionModelHandler.cpp
+++ b/Engine/Plugins/VirtualProduction/CameraCalibrationCore/Source/CameraCalibrationCore/Private/SphericalLensDistortionModelHandler.cpp
@@ -11,15 +11,45 @@ void USphericalLensDistortionModelHandler::InitializeHandler()
 	LensModelClass = USphericalLensModel::StaticClass();
 }
 
-FVector2D USphericalLensDistortionModelHandler::ComputeDistortedUV(const FVector2D& InUndistortedUV) const
+FVector2D USphericalLensDistortionModelHandler::ComputeDistortedUV(const FVector2D& InScreenUV) const
 {
 	// These distances cannot be zero in real-life. If they are, the current distortion state must be bad
 	if ((CurrentState.FocalLengthInfo.FxFy.X == 0.0f) || (CurrentState.FocalLengthInfo.FxFy.Y == 0.0f))
 	{
-		return InUndistortedUV;
+		return InScreenUV;
 	}
+	
+	FVector2D NormalizedDistanceFromImageCenter = (InScreenUV - FVector2D(0.5f, 0.5f)) / CurrentState.FocalLengthInfo.FxFy;
 
-	FVector2D NormalizedDistanceFromImageCenter = (InUndistortedUV - CurrentState.ImageCenter.PrincipalPoint) / CurrentState.FocalLengthInfo.FxFy;
+	const FVector2D DistanceSquared = NormalizedDistanceFromImageCenter * NormalizedDistanceFromImageCenter;
+	const float RSquared = DistanceSquared.X + DistanceSquared.Y;
+
+	// Compute radial distortion
+	const float RadialDistortion = 1.0f + (SphericalParameters.K1 * RSquared) + (SphericalParameters.K2 * RSquared * RSquared) + (SphericalParameters.K3 * RSquared * RSquared * RSquared);
+	
+	// Compute tangential distortion
+	const FVector2D TangentialDistortion = 
+	{
+		(SphericalParameters.P2 * (RSquared + (2.0f * DistanceSquared.X))) + (2.0f * SphericalParameters.P1 * NormalizedDistanceFromImageCenter.X * NormalizedDistanceFromImageCenter.Y),
+		(SphericalParameters.P1 * (RSquared + (2.0f * DistanceSquared.Y))) + (2.0f * SphericalParameters.P2 * NormalizedDistanceFromImageCenter.X * NormalizedDistanceFromImageCenter.Y)
+	};
+
+	// Compute total distortion
+	NormalizedDistanceFromImageCenter = NormalizedDistanceFromImageCenter * RadialDistortion + TangentialDistortion;
+
+	const FVector2D DistortedUV = (NormalizedDistanceFromImageCenter * CurrentState.FocalLengthInfo.FxFy) + CurrentState.ImageCenter.PrincipalPoint;
+	return DistortedUV;
+}
+
+FVector2D USphericalLensDistortionModelHandler::ComputeUndistortedUV(const FVector2D& InScreenUV) const
+{
+	// These distances cannot be zero in real-life. If they are, the current distortion state must be bad
+	if ((CurrentState.FocalLengthInfo.FxFy.X == 0.0f) || (CurrentState.FocalLengthInfo.FxFy.Y == 0.0f))
+	{
+		return InScreenUV;
+	}
+
+	FVector2D NormalizedDistanceFromImageCenter = (InScreenUV - CurrentState.ImageCenter.PrincipalPoint) / CurrentState.FocalLengthInfo.FxFy;
 	const FVector2D OriginalDistance = NormalizedDistanceFromImageCenter;
 
 	// Iterative approach to distort an undistorted UV using coefficients that were designed to undistort
diff --git a/Engine/Plugins/VirtualProduction/CameraCalibrationCore/Source/CameraCalibrationCore/Public/AnamorphicLensDistortionModelHandler.h b/Engine/Plugins/VirtualProduction/CameraCalibrationCore/Source/CameraCalibrationCore/Public/AnamorphicLensDistortionModelHandler.h
--- a/Engine/Plugins/VirtualProduction/CameraCalibrationCore/Source/CameraCalibrationCore/Public/AnamorphicLensDistortionModelHandler.h
+++ b/Engine/Plugins/VirtualProduction/CameraCalibrationCore/Source/CameraCalibrationCore/Public/AnamorphicLensDistortionModelHandler.h
@@ -18,6 +18,7 @@ protected:
 	//~ Begin ULensDistortionModelHandlerBase interface
 	virtual void InitializeHandler() override;
 	virtual FVector2D ComputeDistortedUV(const FVector2D& InScreenUV) const override;
+	virtual FVector2D ComputeUndistortedUV(const FVector2D& InScreenUV) const override;
 	virtual void InitDistortionMaterials() override;
 	virtual void UpdateMaterialParameters() override;
 	virtual void InterpretDistortionParameters() override;
diff --git a/Engine/Plugins/VirtualProduction/CameraCalibrationCore/Source/CameraCalibrationCore/Public/CameraCalibrationSubsystem.h b/Engine/Plugins/VirtualProduction/CameraCalibrationCore/Source/CameraCalibrationCore/Public/CameraCalibrationSubsystem.h
--- a/Engine/Plugins/VirtualProduction/CameraCalibrationCore/Source/CameraCalibrationCore/Public/CameraCalibrationSubsystem.h
+++ b/Engine/Plugins/VirtualProduction/CameraCalibrationCore/Source/CameraCalibrationCore/Public/CameraCalibrationSubsystem.h
@@ -119,7 +119,7 @@ public:
 	void UnregisterOverlayMaterial(const FName& MaterialName);
 
 	/** Adds the input distortion state and blending parameters to the Lens Distortion Scene View Extension for the input camera actor */
-	void SetLensDistortionSVEState(ACameraActor* CameraActor, FDisplacementMapBlendingParams DistortionState);
+	void SetLensDistortionSVEState(ACameraActor* CameraActor, FDisplacementMapBlendingParams DistortionState, ULensDistortionModelHandlerBase* LensDistortionHandler = nullptr);
 
 	/** Removes all distortion state from the Lens Distortion Scene View Extension for the input camera actor */
 	void ClearLensDistortionSVEState(ACameraActor* CameraActor);
diff --git a/Engine/Plugins/VirtualProduction/CameraCalibrationCore/Source/CameraCalibrationCore/Public/LensDistortionModelHandlerBase.h b/Engine/Plugins/VirtualProduction/CameraCalibrationCore/Source/CameraCalibrationCore/Public/LensDistortionModelHandlerBase.h
--- a/Engine/Plugins/VirtualProduction/CameraCalibrationCore/Source/CameraCalibrationCore/Public/LensDistortionModelHandlerBase.h
+++ b/Engine/Plugins/VirtualProduction/CameraCalibrationCore/Source/CameraCalibrationCore/Public/LensDistortionModelHandlerBase.h
@@ -107,6 +107,9 @@ public:
 	/** Use the current distortion state to compute the overscan factor needed such that all distorted UVs will fall into the valid range of [0,1] */
 	float ComputeOverscanFactor() const;
 
+	/** Use the current distortion state to compute the overscan factor needed when undistorting a distorted image */
+	float ComputeInverseOverscanFactor() const;
+	
 	/** Computes the distorted version of UndistortedUVs based on the current state */
 	TArray<FVector2D> GetDistortedUVs(TConstArrayView<FVector2D> UndistortedUVs) const;
 
@@ -129,6 +132,9 @@ protected:
 	/** Use the current distortion state to compute the distortion position of an input UV coordinate */
 	virtual FVector2D ComputeDistortedUV(const FVector2D& InScreenUV) const PURE_VIRTUAL(ULensDistortionModelHandlerBase::ComputeDistortedUV, return FVector2D::ZeroVector;);
 
+	/** Use the current distortion state to compute the undistortion position of an input UV coordinate */
+	virtual FVector2D ComputeUndistortedUV(const FVector2D& InScreenUV) const PURE_VIRTUAL(ULensDistortionModelHandlerBase::ComputeUndistortedUV, return FVector2D::ZeroVector;);
+	
 	/** Create the distortion MIDs */
 	virtual void InitDistortionMaterials() PURE_VIRTUAL(ULensDistortionModelHandlerBase::InitDistortionMaterials);
 
diff --git a/Engine/Plugins/VirtualProduction/CameraCalibrationCore/Source/CameraCalibrationCore/Public/LensDistortionSceneViewExtension.h b/Engine/Plugins/VirtualProduction/CameraCalibrationCore/Source/CameraCalibrationCore/Public/LensDistortionSceneViewExtension.h
--- a/Engine/Plugins/VirtualProduction/CameraCalibrationCore/Source/CameraCalibrationCore/Public/LensDistortionSceneViewExtension.h
+++ b/Engine/Plugins/VirtualProduction/CameraCalibrationCore/Source/CameraCalibrationCore/Public/LensDistortionSceneViewExtension.h
@@ -13,6 +13,16 @@ class UCameraComponent;
  */
 class FLensDistortionSceneViewExtension : public FSceneViewExtensionBase
 {
+public:
+	/** The distortion model to use when creating distortion maps */
+	enum class EDistortionModel : uint8
+	{
+		None,
+		SphericalDistortion,
+		AnamorphicDistortion,
+		MAX
+	};
+	
 public:
 
 	FLensDistortionSceneViewExtension(const FAutoRegister& AutoRegister);
@@ -23,17 +33,17 @@ public:
 
 public:
 	/** Update the distortion state and blending params for the input camera */
-	void UpdateDistortionState_AnyThread(ACameraActor* CameraActor, FDisplacementMapBlendingParams DistortionState);
+	void UpdateDistortionState_AnyThread(ACameraActor* CameraActor, FDisplacementMapBlendingParams DistortionState, ULensDistortionModelHandlerBase* LensDistortionHandler = nullptr);
 
 	/** Remove the distortion state and blending params for the input camera */
 	void ClearDistortionState_AnyThread(ACameraActor* CameraActor);
 
 private:
 	/** Use the input distortion state to draw a distortion displacement map */
-	void DrawDisplacementMap_RenderThread(FRDGBuilder& GraphBuilder, const FLensDistortionState& CurrentState, float InverseOverscan, float CameraOverscan, FRDGTextureRef& OutDistortionMapWithOverscan);
+	void DrawDisplacementMap_RenderThread(FRDGBuilder& GraphBuilder, const FLensDistortionState& CurrentState, EDistortionModel DistortionModel, float InverseOverscan, float CameraOverscan, const FVector2D& SensorSize, FRDGTextureRef& OutDistortionMapWithOverscan);
 
 	/** Use the input blend parameters to draw multiple displacement maps and blend them together into a final distortion displacement map */
-	void BlendDisplacementMaps_RenderThread(FRDGBuilder& GraphBuilder, const FDisplacementMapBlendingParams& BlendState, float InverseOverscan, float CameraOverscan, FRDGTextureRef& OutDistortionMapWithOverscan);
+	void BlendDisplacementMaps_RenderThread(FRDGBuilder& GraphBuilder, const FDisplacementMapBlendingParams& BlendState, EDistortionModel DistortionModel, float InverseOverscan, float CameraOverscan, const FVector2D& SensorSize, FRDGTextureRef& OutDistortionMapWithOverscan);
 
 	/** Crop the input overscanned distortion map to the original requested resolution */
 	void CropDisplacementMap_RenderThread(FRDGBuilder& GraphBuilder, const FRDGTextureRef& InDistortionMapWithOverscan, FRDGTextureRef& OutDistortionMap);
@@ -41,9 +51,15 @@ private:
 	/** Invert the input distortion map to generate a matching undistortion map (with no overscan) */
 	void InvertDistortionMap_RenderThread(FRDGBuilder& GraphBuilder, const FRDGTextureRef& InDistortionMap, float InInverseOverscan, FRDGTextureRef& OutUndistortionMap);
 
+	/** Gets whether the specified distortion model is forward distorting (generating the distort map first, then inverting for the undistort) or inverse distorting */
+	bool IsDistortionModelForwardDistorting(EDistortionModel InDistortionModel) const;
+
+	/** Computes the overscan needed to invert the distortion map */
+	float GetInverseOverscan(ULensDistortionModelHandlerBase* LensDistortionHandler, EDistortionModel InDistortionModel) const;
+	
 private:
 	/** Map of cameras to their associated distortion state and blending parameters, used to determine if and how displacement maps should be rendered for a specific view */
-	struct Rec { FDisplacementMapBlendingParams Params; UCameraComponent* Component; };
+	struct Rec { FDisplacementMapBlendingParams Params; UCameraComponent* Component; ULensDistortionModelHandlerBase* LensDistortionHandler; };
 	TMap<uint32, Rec> DistortionStateMap;
 
 	/** Critical section to lock access to the distortion state map when potentially being accessed from multiple threads */
diff --git a/Engine/Plugins/VirtualProduction/CameraCalibrationCore/Source/CameraCalibrationCore/Public/SphericalLensDistortionModelHandler.h b/Engine/Plugins/VirtualProduction/CameraCalibrationCore/Source/CameraCalibrationCore/Public/SphericalLensDistortionModelHandler.h
--- a/Engine/Plugins/VirtualProduction/CameraCalibrationCore/Source/CameraCalibrationCore/Public/SphericalLensDistortionModelHandler.h
+++ b/Engine/Plugins/VirtualProduction/CameraCalibrationCore/Source/CameraCalibrationCore/Public/SphericalLensDistortionModelHandler.h
@@ -18,6 +18,7 @@ protected:
 	//~ Begin ULensDistortionModelHandlerBase interface
 	virtual void InitializeHandler() override;
 	virtual FVector2D ComputeDistortedUV(const FVector2D& InScreenUV) const override;
+	virtual FVector2D ComputeUndistortedUV(const FVector2D& InScreenUV) const override;
 	virtual void InitDistortionMaterials() override;
 	virtual void UpdateMaterialParameters() override;
 	virtual void InterpretDistortionParameters() override;
diff --git a/Engine/Plugins/VirtualProduction/LensComponent/Source/LensComponent/Private/LensComponent.cpp b/Engine/Plugins/VirtualProduction/LensComponent/Source/LensComponent/Private/LensComponent.cpp
--- a/Engine/Plugins/VirtualProduction/LensComponent/Source/LensComponent/Private/LensComponent.cpp
+++ b/Engine/Plugins/VirtualProduction/LensComponent/Source/LensComponent/Private/LensComponent.cpp
@@ -209,7 +209,7 @@ void ULensComponent::TickComponent(float DeltaTime, ELevelTick TickType, FActorC
 					{
 						if (ACameraActor* CameraActor = Cast<ACameraActor>(GetOwner()))
 						{
-							SubSystem->SetLensDistortionSVEState(CameraActor, BlendState);
+							SubSystem->SetLensDistortionSVEState(CameraActor, BlendState, LensDistortionHandler);
 						}
 					}
 
