Lens Calibration: Fixed issue with missing ShouldCompilePermutation function in lens distortion SVE shader

diff --git a/Engine/Plugins/VirtualProduction/CameraCalibrationCore/Source/CameraCalibrationCore/Private/LensDistortionSceneViewExtension.cpp b/Engine/Plugins/VirtualProduction/CameraCalibrationCore/Source/CameraCalibrationCore/Private/LensDistortionSceneViewExtension.cpp
--- a/Engine/Plugins/VirtualProduction/CameraCalibrationCore/Source/CameraCalibrationCore/Private/LensDistortionSceneViewExtension.cpp
+++ b/Engine/Plugins/VirtualProduction/CameraCalibrationCore/Source/CameraCalibrationCore/Private/LensDistortionSceneViewExtension.cpp
@@ -82,6 +82,12 @@ class FDrawDistortionDisplacementMapCS : public FGlobalShader
 		SHADER_PARAMETER_STRUCT_INCLUDE(FAnamorphicDistortionParams, AnamorphicDistortionParams)
 		SHADER_PARAMETER_RDG_TEXTURE_UAV(RWTexture2D, OutDistortionMap)
 	END_SHADER_PARAMETER_STRUCT()
+
+	// Called by the engine to determine which permutations to compile for this shader
+	static bool ShouldCompilePermutation(const FGlobalShaderPermutationParameters& Parameters)
+	{
+		return IsFeatureLevelSupported(Parameters.Platform, ERHIFeatureLevel::SM5);
+	}
 };
 
 IMPLEMENT_GLOBAL_SHADER(FDrawDistortionDisplacementMapCS, "/Plugin/CameraCalibrationCore/Private/DrawDisplacementMaps.usf", "MainCS", SF_Compute);
