LensDistortion: Fix race condition/crash with destroyed uobjects in scene view extension render thread function(s).

diff --git a/Engine/Plugins/VirtualProduction/CameraCalibrationCore/Source/CameraCalibrationCore/Private/LensDistortionSceneViewExtension.cpp b/Engine/Plugins/VirtualProduction/CameraCalibrationCore/Source/CameraCalibrationCore/Private/LensDistortionSceneViewExtension.cpp
--- a/Engine/Plugins/VirtualProduction/CameraCalibrationCore/Source/CameraCalibrationCore/Private/LensDistortionSceneViewExtension.cpp
+++ b/Engine/Plugins/VirtualProduction/CameraCalibrationCore/Source/CameraCalibrationCore/Private/LensDistortionSceneViewExtension.cpp
@@ -30,7 +30,22 @@ FLensDistortionSceneViewExtension::FLensDistortionSceneViewExtension(const FAuto
 void FLensDistortionSceneViewExtension::UpdateDistortionState_AnyThread(ACameraActor* CameraActor, FDisplacementMapBlendingParams DistortionState, ULensDistortionModelHandlerBase* LensDistortionHandler)
 {
 	FScopeLock ScopeLock(&DistortionStateMapCriticalSection);
-	DistortionStateMap.Add(CameraActor->GetUniqueID(), {DistortionState, CameraActor->GetCameraComponent(), LensDistortionHandler});
+
+	FCameraDistortionProxy CamProxy;
+	CamProxy.Params = DistortionState;
+	CamProxy.LensDistortionHandler = LensDistortionHandler;
+
+	if (UCameraComponent* Component = CameraActor->GetCameraComponent())
+	{
+		CamProxy.CameraOverscan = Component->Overscan + 1.0f;
+		
+		if (UCineCameraComponent* CineCameraComponent = Cast<UCineCameraComponent>(Component))
+		{
+			CamProxy.FilmbackSettings = CineCameraComponent->Filmback;
+		}
+	}
+	
+	DistortionStateMap.Add(CameraActor->GetUniqueID(), MoveTemp(CamProxy));
 }
 
 void FLensDistortionSceneViewExtension::ClearDistortionState_AnyThread(ACameraActor* CameraActor)
@@ -445,30 +460,25 @@ float FLensDistortionSceneViewExtension::GetInverseOverscan(ULensDistortionModel
 void FLensDistortionSceneViewExtension::PreRenderView_RenderThread(FRDGBuilder& GraphBuilder, FSceneView& InView)
 {
 	FScopeLock ScopeLock(&DistortionStateMapCriticalSection);
-	Rec* Rec = DistortionStateMap.Find(InView.ViewActor.ActorUniqueId);
-	if (!Rec)
+	FCameraDistortionProxy* Proxy = DistortionStateMap.Find(InView.ViewActor.ActorUniqueId);
+	if (!Proxy)
 	{
 		return;
 	}
 
-	FCameraFilmbackSettings CameraFilmback;
-	if (UCineCameraComponent* CineCameraComponent = Cast<UCineCameraComponent>(Rec->Component))
-	{
-		CameraFilmback = CineCameraComponent->Filmback;
-	}
+	FVector2D SensorSize = FVector2D(Proxy->FilmbackSettings.SensorWidth, Proxy->FilmbackSettings.SensorHeight);
+	FDisplacementMapBlendingParams BlendState = Proxy->Params;
 
-	FVector2D SensorSize = FVector2D(CameraFilmback.SensorWidth, CameraFilmback.SensorHeight);
-	float CameraOverscan = Rec->Component ? (Rec->Component->Overscan + 1.0f) : 1.0f;
-	FDisplacementMapBlendingParams BlendState = Rec->Params;
+	TStrongObjectPtr<ULensDistortionModelHandlerBase> ModelHandler = Proxy->LensDistortionHandler.Pin();
 	
 	EDistortionModel DistortionModel = EDistortionModel::None;
-	if (Rec->LensDistortionHandler)
+	if (ModelHandler.IsValid())
 	{
-		if (Rec->LensDistortionHandler->GetLensModelClass() == USphericalLensModel::StaticClass())
+		if (ModelHandler->GetLensModelClass() == USphericalLensModel::StaticClass())
 		{
 			DistortionModel = EDistortionModel::SphericalDistortion;
 		}
-		else if (Rec->LensDistortionHandler->GetLensModelClass() == UAnamorphicLensModel::StaticClass())
+		else if (ModelHandler->GetLensModelClass() == UAnamorphicLensModel::StaticClass())
 		{
 			DistortionModel = EDistortionModel::AnamorphicDistortion;
 		}
@@ -503,7 +513,7 @@ void FLensDistortionSceneViewExtension::PreRenderView_RenderThread(FRDGBuilder&
 	ViewDistortionLUT.UndistortingDisplacementTexture = GraphBuilder.CreateTexture(UndistortionMapDesc, TEXT("UndistortionDisplacementMap"));
 
 	// In order to guarantee that we can generate a complete undistortion map, the distortion map we invert needs to have some overscan 
-	float InverseOverscan = FMath::Clamp(GetInverseOverscan(Rec->LensDistortionHandler, DistortionModel), 0.0f, 2.0f);
+	float InverseOverscan = FMath::Clamp(GetInverseOverscan(ModelHandler.Get(), DistortionModel), 0.0f, 2.0f);
 	
 	// Adjust the overscan resolution to be square, with each side being a multiple of 8
 	const FIntPoint OverscanResolution = FIntPoint(FMath::CeilToInt(InverseOverscan * 32) * 8);
@@ -522,7 +532,7 @@ void FLensDistortionSceneViewExtension::PreRenderView_RenderThread(FRDGBuilder&
 	FRDGTextureRef ForwardDistortionMap = bForwardDistort ? ViewDistortionLUT.DistortingDisplacementTexture : ViewDistortionLUT.UndistortingDisplacementTexture;
 	FRDGTextureRef InverseDistortionMap = bForwardDistort ? ViewDistortionLUT.UndistortingDisplacementTexture : ViewDistortionLUT.DistortingDisplacementTexture;
 	
-	BlendDisplacementMaps_RenderThread(GraphBuilder, BlendState, DistortionModel, InverseOverscan, CameraOverscan, SensorSize, DistortionMapWithOverscan);
+	BlendDisplacementMaps_RenderThread(GraphBuilder, BlendState, DistortionModel, InverseOverscan, Proxy->CameraOverscan, SensorSize, DistortionMapWithOverscan);
 	InvertDistortionMap_RenderThread(GraphBuilder, DistortionMapWithOverscan, InverseOverscan, InverseDistortionMap);
 	CropDisplacementMap_RenderThread(GraphBuilder, DistortionMapWithOverscan, ForwardDistortionMap);
 
diff --git a/Engine/Plugins/VirtualProduction/CameraCalibrationCore/Source/CameraCalibrationCore/Public/LensDistortionSceneViewExtension.h b/Engine/Plugins/VirtualProduction/CameraCalibrationCore/Source/CameraCalibrationCore/Public/LensDistortionSceneViewExtension.h
--- a/Engine/Plugins/VirtualProduction/CameraCalibrationCore/Source/CameraCalibrationCore/Public/LensDistortionSceneViewExtension.h
+++ b/Engine/Plugins/VirtualProduction/CameraCalibrationCore/Source/CameraCalibrationCore/Public/LensDistortionSceneViewExtension.h
@@ -58,9 +58,16 @@ private:
 	float GetInverseOverscan(ULensDistortionModelHandlerBase* LensDistortionHandler, EDistortionModel InDistortionModel) const;
 	
 private:
+
+	struct FCameraDistortionProxy
+	{
+		FDisplacementMapBlendingParams Params = {};
+		float CameraOverscan = 1.0f;
+		FCameraFilmbackSettings FilmbackSettings = {};
+		TWeakObjectPtr<ULensDistortionModelHandlerBase> LensDistortionHandler = nullptr;
+	};
 	/** Map of cameras to their associated distortion state and blending parameters, used to determine if and how displacement maps should be rendered for a specific view */
-	struct Rec { FDisplacementMapBlendingParams Params; UCameraComponent* Component; ULensDistortionModelHandlerBase* LensDistortionHandler; };
-	TMap<uint32, Rec> DistortionStateMap;
+	TMap<uint32, FCameraDistortionProxy> DistortionStateMap;
 
 	/** Critical section to lock access to the distortion state map when potentially being accessed from multiple threads */
 	mutable FCriticalSection DistortionStateMapCriticalSection;
