Lens Calibration: Added support for ST maps in the lens distortion SVE.

diff --git a/Engine/Plugins/VirtualProduction/CameraCalibrationCore/Shaders/Private/FillSTDisplacementMap.usf b/Engine/Plugins/VirtualProduction/CameraCalibrationCore/Shaders/Private/FillSTDisplacementMap.usf
new file mode 100644
--- /dev/null
+++ b/Engine/Plugins/VirtualProduction/CameraCalibrationCore/Shaders/Private/FillSTDisplacementMap.usf
@@ -0,0 +1,50 @@
+﻿// Copyright Epic Games, Inc. All Rights Reserved.
+
+/*=============================================================================
+	FillSTDisplacementMap.usf
+=============================================================================*/
+
+#include "/Engine/Public/Platform.ush"
+
+uint2 GridDimensions;
+
+float Overscan;
+float2 PixelToUV;
+float2 PixelToOverscannedUV;
+
+Texture2D UndisplacementMap;
+SamplerState UndisplacementMapSampler;
+
+float2 ToOverscannedSpace(float2 UV)
+{
+	return (UV - float2(0.5, 0.5)) / Overscan + float2(0.5, 0.5);
+}
+
+void MainVS(
+	in uint VertexId : SV_VertexID,
+	out noperspective float4 OutTexCoord : TEXCOORD0,
+	out float4 OutPosition : SV_POSITION)
+{
+	// Compute UV from vertex ID and rectangle layout
+	float2 UnoverscannedUV = float2(VertexId % (GridDimensions.x + 1), VertexId / (GridDimensions.x + 1)) / float2(GridDimensions) + 0.5 * PixelToUV;
+	float2 ViewportUV = ToOverscannedSpace(UnoverscannedUV);
+	float2 DisplacedUV = ToOverscannedSpace(UnoverscannedUV + UndisplacementMap.SampleLevel(UndisplacementMapSampler, UnoverscannedUV, 0));
+	
+	// Write the original position in original and overscanned space to the texture coordinate
+	OutTexCoord.xy = ViewportUV;
+	OutTexCoord.zw = UnoverscannedUV;
+
+	// Output the distorted vertex in clip space
+	OutPosition = float4(((DisplacedUV * 2.0) - 1.0) * int2(1, -1), 0.0, 1.0);
+}
+
+void MainPS(
+	in noperspective float4 InTexCoord : TEXCOORD0,
+	in float4 SvPosition : SV_POSITION,
+	out float4 OutColor : SV_Target0)
+{
+	float2 CurrentUV = SvPosition.xy * PixelToOverscannedUV;
+	float2 OriginalUV = InTexCoord.xy;
+
+	OutColor= float4((OriginalUV - CurrentUV), 0, 1);
+}
\ No newline at end of file
diff --git a/Engine/Plugins/VirtualProduction/CameraCalibrationCore/Shaders/Private/RecenterSTDisplacementMaps.usf b/Engine/Plugins/VirtualProduction/CameraCalibrationCore/Shaders/Private/RecenterSTDisplacementMaps.usf
new file mode 100644
--- /dev/null
+++ b/Engine/Plugins/VirtualProduction/CameraCalibrationCore/Shaders/Private/RecenterSTDisplacementMaps.usf
@@ -0,0 +1,26 @@
+﻿// Copyright Epic Games, Inc. All Rights Reserved.
+
+#include "/Engine/Private/Common.ush"
+#include "/Engine/Private/ScreenPass.ush"
+
+SCREEN_PASS_TEXTURE_VIEWPORT(Output)
+
+float Overscan;
+float2 STMapInvSize;
+Texture2D<float2> STMap;
+SamplerState STMapSampler;
+
+bool ShouldCrop(float2 UV)
+{
+	// Add a small border of pixels beyond the normal crop to ensure the boundary doesn't show up in TSR
+	float2 OverscannedInvSize = Overscan * STMapInvSize;
+	return UV.x + 2 * OverscannedInvSize.x < 0.0 || UV.x - 2 * OverscannedInvSize.x > 1.0 || UV.y + 2 * OverscannedInvSize.y < 0.0 || UV.y - 2 * OverscannedInvSize.y > 1.0;
+}
+
+void MainPS(noperspective float4 UVAndScreenPos : TEXCOORD0, float4 SvPosition : SV_POSITION, out float4 OutColor : SV_Target0)
+{
+	float2 UnoverscannedUV = (UVAndScreenPos.xy - float2(0.5, 0.5)) * Overscan + float2(0.5, 0.5);
+	bool bCropPixel = ShouldCrop(UnoverscannedUV);
+	float2 Displacement = STMap.SampleLevel(STMapSampler, UnoverscannedUV, 0);
+	OutColor = float4(select(bCropPixel, float2(0.0, 0.0), Displacement / Overscan), 0, 1);
+}
\ No newline at end of file
diff --git a/Engine/Plugins/VirtualProduction/CameraCalibrationCore/Source/CameraCalibrationCore/Private/LensDistortionSceneViewExtension.cpp b/Engine/Plugins/VirtualProduction/CameraCalibrationCore/Source/CameraCalibrationCore/Private/LensDistortionSceneViewExtension.cpp
--- a/Engine/Plugins/VirtualProduction/CameraCalibrationCore/Source/CameraCalibrationCore/Private/LensDistortionSceneViewExtension.cpp
+++ b/Engine/Plugins/VirtualProduction/CameraCalibrationCore/Source/CameraCalibrationCore/Private/LensDistortionSceneViewExtension.cpp
@@ -6,12 +6,14 @@
 #include "Camera/CameraActor.h"
 #include "Camera/CameraComponent.h"
 #include "Containers/DynamicRHIResourceArray.h"
+#include "Engine/TextureRenderTarget2D.h"
 #include "GlobalShader.h"
 #include "PostProcess/LensDistortion.h"
 #include "RHIResourceUtils.h"
 #include "ScreenPass.h"
 #include "ShaderParameterStruct.h"
 #include "SystemTextures.h"
+#include "TextureResource.h"
 #include "Models/AnamorphicLensModel.h"
 #include "Models/SphericalLensModel.h"
 
@@ -182,6 +184,53 @@ public:
 IMPLEMENT_GLOBAL_SHADER(FInvertDisplacementVS, "/Plugin/CameraCalibrationCore/Private/InvertDisplacementMap.usf", "MainVS", SF_Vertex);
 IMPLEMENT_GLOBAL_SHADER(FInvertDisplacementPS, "/Plugin/CameraCalibrationCore/Private/InvertDisplacementMap.usf", "MainPS", SF_Pixel);
 
+BEGIN_SHADER_PARAMETER_STRUCT(FFillSTDisplacementMapParameters, )
+	SHADER_PARAMETER(FIntPoint, GridDimensions)
+	SHADER_PARAMETER(FVector2f, PixelToUV)
+	SHADER_PARAMETER(FVector2f, PixelToOverscannedUV)
+	SHADER_PARAMETER(float, Overscan)
+	SHADER_PARAMETER_RDG_TEXTURE_SRV(Texture2D, UndisplacementMap)
+	SHADER_PARAMETER_SAMPLER(SamplerState, UndisplacementMapSampler)
+	RENDER_TARGET_BINDING_SLOTS()
+END_SHADER_PARAMETER_STRUCT()
+
+class FFillSTDisplacementMapVS : public FGlobalShader
+{
+public:
+	DECLARE_GLOBAL_SHADER(FFillSTDisplacementMapVS);
+	SHADER_USE_PARAMETER_STRUCT(FFillSTDisplacementMapVS, FGlobalShader);
+	using FParameters = FFillSTDisplacementMapParameters;
+};
+
+class FFillSTDisplacementMapPS : public FGlobalShader
+{
+public:
+	DECLARE_GLOBAL_SHADER(FFillSTDisplacementMapPS);
+	SHADER_USE_PARAMETER_STRUCT(FFillSTDisplacementMapPS, FGlobalShader);
+	using FParameters = FFillSTDisplacementMapParameters;
+};
+
+IMPLEMENT_GLOBAL_SHADER(FFillSTDisplacementMapVS, "/Plugin/CameraCalibrationCore/Private/FillSTDisplacementMap.usf", "MainVS", SF_Vertex);
+IMPLEMENT_GLOBAL_SHADER(FFillSTDisplacementMapPS, "/Plugin/CameraCalibrationCore/Private/FillSTDisplacementMap.usf", "MainPS", SF_Pixel);
+
+BEGIN_SHADER_PARAMETER_STRUCT(FRecenterSTDisplacementMapParameters, )
+	SHADER_PARAMETER(float, Overscan)
+	SHADER_PARAMETER(FVector2f, STMapInvSize)
+	SHADER_PARAMETER_RDG_TEXTURE_SRV(Texture2D, STMap)
+	SHADER_PARAMETER_SAMPLER(SamplerState, STMapSampler)
+	RENDER_TARGET_BINDING_SLOTS()
+END_SHADER_PARAMETER_STRUCT()
+
+class FRecenterSTDisplacementMapPS : public FGlobalShader
+{
+public:
+	DECLARE_GLOBAL_SHADER(FRecenterSTDisplacementMapPS);
+	SHADER_USE_PARAMETER_STRUCT(FRecenterSTDisplacementMapPS, FGlobalShader);
+	using FParameters = FRecenterSTDisplacementMapParameters;
+};
+
+IMPLEMENT_GLOBAL_SHADER(FRecenterSTDisplacementMapPS, "/Plugin/CameraCalibrationCore/Private/RecenterSTDisplacementMaps.usf", "MainPS", SF_Pixel);
+
 void FLensDistortionSceneViewExtension::DrawDisplacementMap_RenderThread(FRDGBuilder& GraphBuilder, const FLensDistortionState& CurrentState, EDistortionModel DistortionModel, float InverseOverscan, float CameraOverscan, const FVector2D& SensorSize, FRDGTextureRef& OutDistortionMapWithOverscan)
 {
 	if (CurrentState.DistortionInfo.Parameters.IsEmpty())
@@ -396,6 +445,94 @@ void FLensDistortionSceneViewExtension::InvertDistortionMap_RenderThread(FRDGBui
 
 			const uint32 Width = GridDimensions.X;
 			const uint32 Height = GridDimensions.Y;
+			const uint32 NumVertices = (Width + 1) * (Height + 1);
+			const uint32 NumTriangles = Width * Height * 2;
+			const uint32 NumIndices = NumTriangles * 3;
+
+			IndexBuffer.AddUninitialized(NumIndices);
+
+			uint32* IndexBufferData = (uint32*)IndexBuffer.GetData();
+
+			for (uint32 IndexY = 0; IndexY < Height; ++IndexY)
+			{
+				for (uint32 IndexX = 0; IndexX < Width; ++IndexX)
+				{
+					// left top to bottom right in reading order
+					uint32 Index00 = IndexX + IndexY * (Width + 1);
+					uint32 Index10 = Index00 + 1;
+					uint32 Index01 = Index00 + (Width + 1);
+					uint32 Index11 = Index01 + 1;
+
+					// triangle A
+					*IndexBufferData++ = Index00;
+					*IndexBufferData++ = Index01;
+					*IndexBufferData++ = Index10;
+
+					// triangle B
+					*IndexBufferData++ = Index11;
+					*IndexBufferData++ = Index10;
+					*IndexBufferData++ = Index01;
+				}
+			}
+
+			// Create index buffer. Fill buffer with initial data upon creation
+			FBufferRHIRef IndexBufferRHI = UE::RHIResourceUtils::CreateIndexBufferFromArray(RHICmdList, TEXT("InvertDistortionMapIndexBuffer"), EBufferUsageFlags::Static, MakeConstArrayView(IndexBuffer));
+
+			RHICmdList.DrawIndexedPrimitive(IndexBufferRHI, 0, 0, NumVertices, 0, NumTriangles, 1);
+		});
+}
+
+void FLensDistortionSceneViewExtension::FillSTDisplacementMap_RenderThread(FRDGBuilder& GraphBuilder, const FRDGTextureRef& InUndisplacementMap, float InOverscan, FRDGTextureRef& OutFilledDisplacementMap)
+{
+	FFillSTDisplacementMapParameters* PassParameters = GraphBuilder.AllocParameters<FFillSTDisplacementMapParameters>();
+
+	FScreenPassRenderTarget Output;
+	Output.Texture = OutFilledDisplacementMap;
+	Output.ViewRect = FIntRect(FIntPoint(0, 0), OutFilledDisplacementMap->Desc.Extent);
+	Output.LoadAction = ERenderTargetLoadAction::EClear;
+	Output.UpdateVisualizeTextureExtent();
+	
+	const int32 NumSquares = FMath::Clamp(CVarLensDistortionInvertGridDensity.GetValueOnRenderThread(), 64, 255);
+	const FIntPoint GridDimensions = FIntPoint(NumSquares);
+
+	PassParameters->GridDimensions = GridDimensions;
+
+	PassParameters->RenderTargets[0] = Output.GetRenderTargetBinding();
+
+	PassParameters->UndisplacementMap = GraphBuilder.CreateSRV(InUndisplacementMap);
+	PassParameters->UndisplacementMapSampler = TStaticSamplerState<SF_Bilinear, AM_Clamp, AM_Clamp, AM_Clamp>::GetRHI();
+	
+	PassParameters->Overscan = InOverscan;
+	PassParameters->PixelToUV = FVector2f(1.0f) / FVector2f(InUndisplacementMap->Desc.Extent);
+	PassParameters->PixelToOverscannedUV = FVector2f(1.0f) / FVector2f(OutFilledDisplacementMap->Desc.Extent);
+
+	GraphBuilder.AddPass(
+		RDG_EVENT_NAME("FillSTDisplacementMap"),
+		PassParameters,
+		ERDGPassFlags::Raster,
+		[PassParameters, GridDimensions, Output](FRDGAsyncTask, FRHICommandList& RHICmdList)
+		{
+			RHICmdList.SetViewport(Output.ViewRect.Min.X, Output.ViewRect.Min.Y, 0.0f, Output.ViewRect.Max.X, Output.ViewRect.Max.Y, 1.0f);
+
+			TShaderMapRef<FFillSTDisplacementMapVS> VertexShader(GetGlobalShaderMap(GMaxRHIFeatureLevel));
+			TShaderMapRef<FFillSTDisplacementMapPS> PixelShader(GetGlobalShaderMap(GMaxRHIFeatureLevel));
+
+			SetScreenPassPipelineState(RHICmdList, FScreenPassPipelineState(VertexShader, PixelShader));
+			SetShaderParameters(RHICmdList, VertexShader, VertexShader.GetVertexShader(), *PassParameters);
+			SetShaderParameters(RHICmdList, PixelShader, PixelShader.GetPixelShader(), *PassParameters);
+
+			FRHIBatchedShaderParameters& BatchedParameters = RHICmdList.GetScratchShaderParameters();
+			RHICmdList.SetBatchedShaderParameters(VertexShader.GetVertexShader(), BatchedParameters);
+
+			// No vertex buffer is needed because we compute it in the VS
+			RHICmdList.SetStreamSource(0, NULL, 0);
+
+			// The following code for setting up this index buffer based on FTesselatedScreenRectangleIndexBuffer::InitRHI()
+			TArray<uint32> IndexBuffer;
+
+			const uint32 Width = GridDimensions.X;
+			const uint32 Height = GridDimensions.Y;
+			const uint32 NumVertices = (Width + 1) * (Height + 1);
 			const uint32 NumTriangles = Width * Height * 2;
 			const uint32 NumIndices = NumTriangles * 3;
 
@@ -428,8 +565,50 @@ void FLensDistortionSceneViewExtension::InvertDistortionMap_RenderThread(FRDGBui
 			// Create index buffer. Fill buffer with initial data upon creation
 			FBufferRHIRef IndexBufferRHI = UE::RHIResourceUtils::CreateIndexBufferFromArray(RHICmdList, TEXT("InvertDistortionMapIndexBuffer"), EBufferUsageFlags::Static, MakeConstArrayView(IndexBuffer));
 
-			RHICmdList.DrawIndexedPrimitive(IndexBufferRHI, 0, 0, NumIndices, 0, NumTriangles, 1);
+			RHICmdList.DrawIndexedPrimitive(IndexBufferRHI, 0, 0, NumVertices, 0, NumTriangles, 1);
+		});
+}
+
+void FLensDistortionSceneViewExtension::RecenterSTDisplacementMap_RenderThread(FRDGBuilder& GraphBuilder, const FRDGTextureRef& InDisplacementMap, float InOverscan, FRDGTextureRef& OutRecenteredDisplacementMap)
+{
+	FScreenPassRenderTarget Output;
+	Output.Texture = OutRecenteredDisplacementMap;
+	Output.ViewRect = FIntRect(FIntPoint(0, 0), OutRecenteredDisplacementMap->Desc.Extent);
+	Output.LoadAction = ERenderTargetLoadAction::EClear;
+	Output.UpdateVisualizeTextureExtent();
+	
+	FRecenterSTDisplacementMapPS::FParameters* PassParameters = GraphBuilder.AllocParameters<FRecenterSTDisplacementMapPS::FParameters>();
+	PassParameters->STMapInvSize = FVector2f(1.0f / InDisplacementMap->Desc.Extent.X, 1.0f / InDisplacementMap->Desc.Extent.Y);
+	PassParameters->STMap = GraphBuilder.CreateSRV(InDisplacementMap);
+	PassParameters->STMapSampler = TStaticSamplerState<SF_Bilinear, AM_Clamp, AM_Clamp, AM_Clamp>::GetRHI();
+	PassParameters->Overscan = InOverscan;
+	PassParameters->RenderTargets[0] = Output.GetRenderTargetBinding();
+	
+	FGlobalShaderMap* GlobalShaderMap = GetGlobalShaderMap(GMaxRHIFeatureLevel);
+	TShaderMapRef<FScreenPassVS> ScreenPassVS(GlobalShaderMap);
+	TShaderMapRef<FRecenterSTDisplacementMapPS> DistortImageShader(GlobalShaderMap);
+
+	FRHIBlendState* DefaultBlendState = FScreenPassPipelineState::FDefaultBlendState::GetRHI();
+	const FScreenPassTextureViewport OutViewport(OutRecenteredDisplacementMap->Desc.Extent);
+
+	GraphBuilder.AddPass(
+		RDG_EVENT_NAME("RecenterSTDisplacementMap"),
+		PassParameters,
+		ERDGPassFlags::Raster | ERDGPassFlags::NeverCull,
+		[ScreenPassVS, DistortImageShader, OutViewport, PassParameters, DefaultBlendState](FRDGAsyncTask, FRHICommandList& RHICmdList)
+	{
+		DrawScreenPass(
+			RHICmdList,
+			FScreenPassViewInfo(),
+			OutViewport,
+			OutViewport,
+			FScreenPassPipelineState(ScreenPassVS, DistortImageShader, DefaultBlendState),
+			EScreenPassDrawFlags::None,
+			[&](FRHICommandList&)
+		{
+			SetShaderParameters(RHICmdList, DistortImageShader, DistortImageShader.GetPixelShader(), *PassParameters);
 		});
+	});
 }
 
 bool FLensDistortionSceneViewExtension::IsDistortionModelForwardDistorting(EDistortionModel InDistortionModel) const
@@ -474,16 +653,71 @@ void FLensDistortionSceneViewExtension::PreRenderView_RenderThread(FRDGBuilder&
 	EDistortionModel DistortionModel = EDistortionModel::None;
 	if (ModelHandler.IsValid())
 	{
-		if (ModelHandler->GetLensModelClass() == USphericalLensModel::StaticClass())
+		if (Proxy->Params.States[0].DistortionInfo.Parameters.Num())
 		{
-			DistortionModel = EDistortionModel::SphericalDistortion;
+			if (ModelHandler->GetLensModelClass() == USphericalLensModel::StaticClass())
+			{
+				DistortionModel = EDistortionModel::SphericalDistortion;
+			}
+			else if (ModelHandler->GetLensModelClass() == UAnamorphicLensModel::StaticClass())
+			{
+				DistortionModel = EDistortionModel::AnamorphicDistortion;
+			}
 		}
-		else if (ModelHandler->GetLensModelClass() == UAnamorphicLensModel::StaticClass())
+		else
 		{
-			DistortionModel = EDistortionModel::AnamorphicDistortion;
+			DistortionModel = EDistortionModel::STMap;
 		}
 	}
-
+	
+	FLensDistortionLUT ViewDistortionLUT;
+	
+	if (DistortionModel == EDistortionModel::STMap)
+	{
+		// ST Maps generally have their top left corner corresponding to the distortion at the original frustum, not the overscanned frustum, which means that
+		// for cases where the distortion would require overscanned pixels, we must generate those distortion values for the distortion map from the undistortion map
+		// Further, ST maps may have intrinsic inaccuracies that can cause artifacting in TSR, so use a warp grid inversion to generate a fully approximate
+		// distortion ST map from the undistortion map. We also need to recenter the undistortion map to correspond to the overscan, as well as scaling its displacement
+		// by the overscan amount
+		FTextureRHIRef UndistortionMapTextureRef = ModelHandler->GetDistortionDisplacementMap()->GetResource()->GetTexture2DRHI();
+		FRDGTextureRef OriginalUndistortionMap = GraphBuilder.RegisterExternalTexture(CreateRenderTarget(UndistortionMapTextureRef, TEXT("OriginalUndistortionDisplacementMap")));
+
+		FTextureRHIRef DistortionMapTextureRef = ModelHandler->GetUndistortionDisplacementMap()->GetResource()->GetTexture2DRHI();
+		FRDGTextureRef OriginalDistortionMap = GraphBuilder.RegisterExternalTexture(CreateRenderTarget(DistortionMapTextureRef, TEXT("OriginalDistortionDisplacementMap")));
+
+		FRDGTextureDesc FilledDistortionMapDesc = FRDGTextureDesc::Create2D(
+			OriginalDistortionMap->Desc.Extent,
+			PF_G32R32F,
+			FClearValueBinding::Black,
+			TexCreate_ShaderResource | TexCreate_RenderTargetable);
+
+		FRDGTextureDesc FilledUndistortionMapDesc = FRDGTextureDesc::Create2D(
+			OriginalUndistortionMap->Desc.Extent,
+			PF_G32R32F,
+			FClearValueBinding::Black,
+			TexCreate_ShaderResource | TexCreate_RenderTargetable);
+		
+		ViewDistortionLUT.DistortingDisplacementTexture = GraphBuilder.CreateTexture(FilledDistortionMapDesc, TEXT("DistortionDisplacementMap"));
+		ViewDistortionLUT.UndistortingDisplacementTexture = GraphBuilder.CreateTexture(FilledUndistortionMapDesc, TEXT("UndistortionDisplacementMap"));
+
+		// If there is any overscan, we will need to fill in the overscanned area of the distortion LUT because for ST Maps, the top left corner
+		// of the map corresponds to the top left corner of the original frustum. Perform a mesh warp to invert the undistortion map
+		FRDGTextureRef ForwardTexture = Proxy->CameraOverscan > 1.0f ? OriginalUndistortionMap : OriginalDistortionMap;
+		FRDGTextureRef ForwardTextureOut = Proxy->CameraOverscan > 1.0f ? ViewDistortionLUT.UndistortingDisplacementTexture : ViewDistortionLUT.DistortingDisplacementTexture;
+		FRDGTextureRef InverseTextureOut = Proxy->CameraOverscan > 1.0f ? ViewDistortionLUT.DistortingDisplacementTexture : ViewDistortionLUT.UndistortingDisplacementTexture;
+
+		// Fudge the overscan just a little bit to ensure that the edge of the ST map is not visible during the TSR distortion, as that can cause ghosting and artifacts
+		const float OverscanAdjustment = 0.98f;
+		FillSTDisplacementMap_RenderThread(GraphBuilder, ForwardTexture, Proxy->CameraOverscan * OverscanAdjustment, InverseTextureOut);
+		RecenterSTDisplacementMap_RenderThread(GraphBuilder, ForwardTexture, Proxy->CameraOverscan * OverscanAdjustment, ForwardTextureOut);
+
+		ViewDistortionLUT.DistortionOverscan = Proxy->CameraOverscan * OverscanAdjustment;
+		ViewDistortionLUT.DistortionGridDimensions = FIntPoint(32 * Proxy->CameraOverscan, 20 * Proxy->CameraOverscan);
+		
+		LensDistortion::SetLUTUnsafe(InView, ViewDistortionLUT);
+		return;
+	}
+	
 	const bool bForwardDistort = IsDistortionModelForwardDistorting(DistortionModel);
 	
 	ScopeLock.Unlock();
@@ -503,9 +737,7 @@ void FLensDistortionSceneViewExtension::PreRenderView_RenderThread(FRDGBuilder&
 		PF_G32R32F,
 		FClearValueBinding::Black,
 		TexCreate_ShaderResource | TexCreate_RenderTargetable);
-
-	FLensDistortionLUT ViewDistortionLUT;
-
+	
 	FRDGTextureDesc DistortionMapDesc = bForwardDistort ? ForwardDistortionMapDesc : InverseDistortionMapDesc;
 	FRDGTextureDesc UndistortionMapDesc = bForwardDistort ? InverseDistortionMapDesc : ForwardDistortionMapDesc;
 	
diff --git a/Engine/Plugins/VirtualProduction/CameraCalibrationCore/Source/CameraCalibrationCore/Public/LensDistortionSceneViewExtension.h b/Engine/Plugins/VirtualProduction/CameraCalibrationCore/Source/CameraCalibrationCore/Public/LensDistortionSceneViewExtension.h
--- a/Engine/Plugins/VirtualProduction/CameraCalibrationCore/Source/CameraCalibrationCore/Public/LensDistortionSceneViewExtension.h
+++ b/Engine/Plugins/VirtualProduction/CameraCalibrationCore/Source/CameraCalibrationCore/Public/LensDistortionSceneViewExtension.h
@@ -18,8 +18,15 @@ public:
 	enum class EDistortionModel : uint8
 	{
 		None,
+		
+		/** Distortion maps generated from distortion parameters with a spherical distortion equation */
 		SphericalDistortion,
+
+		/** Distortion maps generated from distortion parameters with an anamorphic distortion equation */
 		AnamorphicDistortion,
+
+		/** Distortion maps generated from UV-based ST maps */
+		STMap,
 		MAX
 	};
 	
@@ -51,6 +58,12 @@ private:
 	/** Invert the input distortion map to generate a matching undistortion map (with no overscan) */
 	void InvertDistortionMap_RenderThread(FRDGBuilder& GraphBuilder, const FRDGTextureRef& InDistortionMap, float InInverseOverscan, FRDGTextureRef& OutUndistortionMap);
 
+	/** Performs a warp grid inverse of the specified distortion map to fill up its inverse to the specified overscan amount */
+	void FillSTDisplacementMap_RenderThread(FRDGBuilder& GraphBuilder, const FRDGTextureRef& InUndisplacementMap, float InOverscan, FRDGTextureRef& OutFilledDisplacementMap);
+
+	/** Recenters the specified distortion map within an overscanned map, filling the overscanned region with zeros */
+	void RecenterSTDisplacementMap_RenderThread(FRDGBuilder& GraphBuilder, const FRDGTextureRef& InDisplacementMap, float InOverscan, FRDGTextureRef& OutRecenteredDisplacementMap);
+	
 	/** Gets whether the specified distortion model is forward distorting (generating the distort map first, then inverting for the undistort) or inverse distorting */
 	bool IsDistortionModelForwardDistorting(EDistortionModel InDistortionModel) const;
 
diff --git a/Engine/Shaders/Private/PostProcessUpscale.usf b/Engine/Shaders/Private/PostProcessUpscale.usf
--- a/Engine/Shaders/Private/PostProcessUpscale.usf
+++ b/Engine/Shaders/Private/PostProcessUpscale.usf
@@ -34,13 +34,7 @@ SamplerState PointSceneColorSampler;
 //in a multiview case, PointSceneColorTexture is an array texture
 Texture2DArray PointSceneColorTextureArray;
 
-// to avoid compiler errors on the PS
-#ifndef TESS_RECT_X
-	#define TESS_RECT_X 1
-#endif
-#ifndef TESS_RECT_Y
-	#define TESS_RECT_Y 1
-#endif
+uint2 GridDimensions;
 
 // Do a smoothstep(x) = 3 x^2 - 2 x^3
 float2 GetSmoothstepUV(float2 LinearUV, float2 TextureSize, float2 TextureInvSize)
@@ -88,10 +82,7 @@ void MainVS(
 {
 	StereoSetupVS(StereoInput, StereoOutput);
 	
-	// amount of Quads in X and Y, from GTesselatedScreenRectangleIndexBuffer
-	uint2 RectangleTessellation = uint2(TESS_RECT_X, TESS_RECT_Y);
-
-	float2 TexCoord = float2(VertexId % (RectangleTessellation.x + 1), VertexId / (RectangleTessellation.x + 1)) / float2(RectangleTessellation);
+	float2 TexCoord = float2(VertexId % (GridDimensions.x + 1), VertexId / (GridDimensions.x + 1)) / float2(GridDimensions);
 
 	// still in 0..1 range
 	float4 Position = float4(TexCoord.x, TexCoord.y, 0, 1);
diff --git a/Engine/Shaders/Private/TemporalSuperResolution/TSRUpdateHistory.usf b/Engine/Shaders/Private/TemporalSuperResolution/TSRUpdateHistory.usf
--- a/Engine/Shaders/Private/TemporalSuperResolution/TSRUpdateHistory.usf
+++ b/Engine/Shaders/Private/TemporalSuperResolution/TSRUpdateHistory.usf
@@ -131,6 +131,7 @@ Texture2D<tsr_ushort2> AntiAliasingTexture;
 	Texture2D<float2> PrevDistortingDisplacementTexture;
 	Texture2D<float2> ResurrectedDistortingDisplacementTexture;
 	Texture2D<float2> UndistortingDisplacementTexture;
+	float DistortionOverscan;
 #endif
 
 FScreenTransform HistoryPixelPosToViewportUV;
@@ -565,13 +566,24 @@ tsr_halfCx2 DownsampleSceneColor(tsr_halfCx2 SceneColor, uint LocalGroupThreadIn
 }
 
 #if CONFIG_LENS_DISTORTION
-float2 ApplyDisplacementTextureOnScreenPos(Texture2D<float2> DisplacementTexture, float2 DestScreenPos)
+/** Samples a displacement texture and crops any values beyond the texture's original bounds */
+float2 SampleDisplacementTexture(Texture2D<float2> DisplacementTexture, float2 TexCoord, float2 ViewportUV)
 {
-	float2 DestViewportUV = ScreenPosToViewportUV(DestScreenPos);
+	// Convert the destination viewport UV to the displacement map's original frustum, and compute whether this value is in or out of that frustum
+	float2 UnoverscannedUV = (ViewportUV - float2(0.5, 0.5)) * DistortionOverscan + float2(0.5, 0.5);
+	bool bCropPixel = UnoverscannedUV.x < 0.0 || UnoverscannedUV.x > 1.0 || UnoverscannedUV.y < 0.0 || UnoverscannedUV.y > 1.0;
 
-	float2 SrcViewportUV = DestViewportUV + DisplacementTexture.SampleLevel(GlobalBilinearClampedSampler, DestViewportUV, 0);
+	return select(bCropPixel, float2(0.0, 0.0), DisplacementTexture.SampleLevel(GlobalBilinearClampedSampler, TexCoord, 0));
+}
+
+float2 ApplyDisplacementTextureOnScreenPos(Texture2D<float2> DisplacementTexture, float2 DistortedScreenPos, float2 ViewportScreenPos)
+{
+	float2 DistortedUV = ScreenPosToViewportUV(DistortedScreenPos);
+	float2 ViewportUV = ScreenPosToViewportUV(ViewportScreenPos);
+	
+	float2 UndistortedUV = DistortedUV + SampleDisplacementTexture(DisplacementTexture, DistortedUV, ViewportUV);
 			
-	return ViewportUVToScreenPos(SrcViewportUV);
+	return ViewportUVToScreenPos(UndistortedUV);
 }
 #endif
 
@@ -656,7 +668,8 @@ void MainCS(
 		UNROLL_N(DPV_PIXEL_PER_LANE)
 		for (uint PixelId = 0; PixelId < DPV_PIXEL_PER_LANE; PixelId++)
 		{
-			float2 LocalSrcViewportUV = dpv_access_pixel(DestViewportUV, PixelId) + UndistortingDisplacementTexture.SampleLevel(GlobalBilinearClampedSampler, dpv_access_pixel(DestViewportUV, PixelId), 0);
+			float2 LocalDestViewportUV = dpv_access_pixel(DestViewportUV, PixelId);			
+			float2 LocalSrcViewportUV = LocalDestViewportUV + SampleDisplacementTexture(UndistortingDisplacementTexture, LocalDestViewportUV, LocalDestViewportUV);
 			
 			SrcViewportUV[0][PixelId] = LocalSrcViewportUV.x;
 			SrcViewportUV[1][PixelId] = LocalSrcViewportUV.y;
@@ -874,12 +887,12 @@ void MainCS(
 		{
 			float2 LocalUndistortedScreenPos = dpv_access_pixel(UndistortedScreenPos, PixelId);
 			float2 LocalPrevScreenPos = dpv_access_pixel(PrevScreenPos, PixelId);
-			float2 LocalSrcScreenPos = ApplyDisplacementTextureOnScreenPos(PrevDistortingDisplacementTexture, LocalPrevScreenPos);
+			float2 LocalSrcScreenPos = ApplyDisplacementTextureOnScreenPos(PrevDistortingDisplacementTexture, LocalPrevScreenPos, LocalUndistortedScreenPos);
 
 			//LocalSrcScreenPos = lerp(LocalUndistortedScreenPos, LocalSrcScreenPos, saturate(OutputPixelVelocity[PixelId] * tsr_half(8.0)));
 			
 			#if CONFIG_HISTORY_RESURRECTION
-				float2 LocalSrcResurrectedScreenPos = ApplyDisplacementTextureOnScreenPos(ResurrectedDistortingDisplacementTexture, LocalPrevScreenPos);
+				float2 LocalSrcResurrectedScreenPos = ApplyDisplacementTextureOnScreenPos(ResurrectedDistortingDisplacementTexture, LocalPrevScreenPos, LocalUndistortedScreenPos);
 				LocalSrcScreenPos = select(bIsHistoryResurrection[PixelId], LocalSrcResurrectedScreenPos, LocalSrcScreenPos);
 			#endif
 
@@ -892,9 +905,17 @@ void MainCS(
 	// Detect whether the history reprojection is successful
 	bool2 bIsOffScreen;
 	bool2 bIsDisoccluded;
+	float Overscan = 1.0f;
+
+#if CONFIG_LENS_DISTORTION
+	Overscan = DistortionOverscan;
+#endif
+	
+	// When there is a distortion overscan, this causes a hard edge where the distortion map ends at its original frustum. Consider
+	// this edge as off screen (which it will be after being cropped at the end of the post process chain) to prevent ghosting from the edge
 	IsOffScreenOrDisoccluded(
 		bCameraCut,
-		PrevScreenPos,
+		PrevScreenPos * Overscan,
 		bIsParallaxRejected,
 		/* out */ bIsOffScreen,
 		/* out */ bIsDisoccluded);
diff --git a/Engine/Source/Runtime/Renderer/Internal/PostProcess/LensDistortion.h b/Engine/Source/Runtime/Renderer/Internal/PostProcess/LensDistortion.h
--- a/Engine/Source/Runtime/Renderer/Internal/PostProcess/LensDistortion.h
+++ b/Engine/Source/Runtime/Renderer/Internal/PostProcess/LensDistortion.h
@@ -18,6 +18,12 @@ struct FLensDistortionLUT
 	/** Resolution fraction of the upscaling happening due to distortion. */
 	float ResolutionFraction = 1.0f;
 
+	/** For distortion maps that don't fill the whole frustum, this is the amount of overscan they require but do not fill */
+	float DistortionOverscan = 1.0f;
+
+	/** The distortion warp grid dimensions to use when distorting during the upscale pass */
+	FIntPoint DistortionGridDimensions = FIntPoint(32, 20);
+	
 	/** Returns whether the displacement is enabled. */
 	bool IsEnabled() const
 	{
diff --git a/Engine/Source/Runtime/Renderer/Private/PostProcess/PostProcessUpscale.cpp b/Engine/Source/Runtime/Renderer/Private/PostProcess/PostProcessUpscale.cpp
--- a/Engine/Source/Runtime/Renderer/Private/PostProcess/PostProcessUpscale.cpp
+++ b/Engine/Source/Runtime/Renderer/Private/PostProcess/PostProcessUpscale.cpp
@@ -5,6 +5,8 @@
 #include "DataDrivenShaderPlatformInfo.h"
 #include "SceneRendering.h"
 #include "PostProcessing.h"
+#include "RHIResourceUtils.h"
+#include "PostProcess/DrawRectangle.h"
 
 namespace
 {
@@ -36,6 +38,7 @@ BEGIN_SHADER_PARAMETER_STRUCT(FUpscaleParameters, )
 	SHADER_PARAMETER_STRUCT(FScreenPassTextureViewportParameters, Output)
 	SHADER_PARAMETER_RDG_TEXTURE(Texture2D, DistortingDisplacementTexture)
 	SHADER_PARAMETER_SAMPLER(SamplerState, DistortingDisplacementSampler)
+	SHADER_PARAMETER(FIntPoint, GridDimensions)
 	SHADER_PARAMETER_RDG_TEXTURE(Texture2D, SceneColorTexture)
 	SHADER_PARAMETER_SAMPLER(SamplerState, SceneColorSampler)
 	SHADER_PARAMETER_RDG_TEXTURE(Texture2D, PointSceneColorTexture)
@@ -84,8 +87,6 @@ public:
 	static void ModifyCompilationEnvironment(const FGlobalShaderPermutationParameters& Parameters, FShaderCompilerEnvironment& OutEnvironment)
 	{
 		FScreenPassVS::ModifyCompilationEnvironment(Parameters, OutEnvironment);
-		OutEnvironment.SetDefine(TEXT("TESS_RECT_X"), FTesselatedScreenRectangleIndexBuffer::Width);
-		OutEnvironment.SetDefine(TEXT("TESS_RECT_Y"), FTesselatedScreenRectangleIndexBuffer::Height);
 	}
 };
 
@@ -143,13 +144,14 @@ FScreenPassTexture ISpatialUpscaler::AddDefaultUpscalePass(
 	const FScreenPassTextureViewport OutputViewport(Output);
 
 	const bool bApplyLensDistortion = LensDistortionLUT.IsEnabled();
-
+	
 	FUpscaleParameters* PassParameters = GraphBuilder.AllocParameters<FUpscaleParameters>();
 	PassParameters->RenderTargets[0] = Output.GetRenderTargetBinding();
 	PassParameters->Input = GetScreenPassTextureViewportParameters(InputViewport);
 	PassParameters->Output = GetScreenPassTextureViewportParameters(OutputViewport);
 	PassParameters->DistortingDisplacementTexture = LensDistortionLUT.DistortingDisplacementTexture;
 	PassParameters->DistortingDisplacementSampler = TStaticSamplerState<SF_Bilinear, AM_Clamp, AM_Clamp, AM_Clamp>::GetRHI();
+	PassParameters->GridDimensions = LensDistortionLUT.DistortionGridDimensions;
 	PassParameters->SceneColorTexture = Inputs.SceneColor.Texture;
 	PassParameters->SceneColorSampler = TStaticSamplerState<SF_Bilinear, AM_Border, AM_Border, AM_Border>::GetRHI();
 	PassParameters->PointSceneColorTexture = Inputs.SceneColor.Texture;
@@ -199,16 +201,76 @@ FScreenPassTexture ISpatialUpscaler::AddDefaultUpscalePass(
 
 		SetShaderParameters(RHICmdList, PixelShader, PixelShader.GetPixelShader(), *PassParameters);
 
-		DrawRectangle(
-			RHICmdList,
-			// Output Rect (RHI viewport relative).
-			0, 0, OutputViewport.Rect.Width(), OutputViewport.Rect.Height(),
-			// Input Rect
-			InputViewport.Rect.Min.X, InputViewport.Rect.Min.Y, InputViewport.Rect.Width(), InputViewport.Rect.Height(),
-			OutputViewport.Rect.Size(),
-			InputViewport.Extent,
-			VertexShader,
-			bApplyLensDistortion ? EDRF_UseTesselatedIndexBuffer : EDRF_UseTriangleOptimization);
+		if (bApplyLensDistortion)
+		{
+			TArray<uint32> IndexBuffer;
+
+			const uint32 Width = PassParameters->GridDimensions.X;
+			const uint32 Height = PassParameters->GridDimensions.Y;
+			const uint32 NumVertices = (Width + 1) * (Height + 1);
+			const uint32 NumTriangles = Width * Height * 2;
+			const uint32 NumIndices = NumTriangles * 3;
+			
+			IndexBuffer.AddUninitialized(NumIndices);
+
+			uint32* Out = (uint32*)IndexBuffer.GetData();
+
+			for(uint32 y = 0; y < Height; ++y)
+			{
+				for(uint32 x = 0; x < Width; ++x)
+				{
+					// left top to bottom right in reading order
+					uint32 Index00 = x  + y * (Width + 1);
+					uint32 Index10 = Index00 + 1;
+					uint32 Index01 = Index00 + (Width + 1);
+					uint32 Index11 = Index01 + 1;
+					
+					// triangle A
+					*Out++ = Index00; *Out++ = Index01; *Out++ = Index10;
+
+					// triangle B
+					*Out++ = Index11; *Out++ = Index10; *Out++ = Index01;
+				}
+			}
+
+			// Create index buffer. Fill buffer with initial data upon creation
+			FBufferRHIRef IndexBufferRHI = UE::RHIResourceUtils::CreateIndexBufferFromArray(RHICmdList, TEXT("LensDistortionIndexBuffer"), EBufferUsageFlags::Static, MakeConstArrayView(IndexBuffer));
+
+			
+			FRHIBatchedShaderParameters& BatchedParameters = RHICmdList.GetScratchShaderParameters();
+			UE::Renderer::PostProcess::SetDrawRectangleParameters(BatchedParameters, VertexShader.GetShader(),
+				0, 0, OutputViewport.Rect.Width(), OutputViewport.Rect.Height(),
+				InputViewport.Rect.Min.X, InputViewport.Rect.Min.Y, InputViewport.Rect.Width(), InputViewport.Rect.Height(),
+				OutputViewport.Rect.Size(),
+				InputViewport.Extent);
+			RHICmdList.SetBatchedShaderParameters(VertexShader.GetVertexShader(), BatchedParameters);
+
+			// no vertex buffer needed as we compute it in VS
+			RHICmdList.SetStreamSource(0, nullptr, 0);
+
+			RHICmdList.DrawIndexedPrimitive(
+				IndexBufferRHI,
+				/*BaseVertexIndex=*/ 0,
+				/*MinIndex=*/ 0,
+				/*NumVertices=*/ NumVertices,
+				/*StartIndex=*/ 0,
+				/*NumPrimitives=*/ NumTriangles,
+				/*NumInstances=*/ 1
+				);
+		}
+		else
+		{
+			DrawRectangle(
+				RHICmdList,
+				// Output Rect (RHI viewport relative).
+				0, 0, OutputViewport.Rect.Width(), OutputViewport.Rect.Height(),
+				// Input Rect
+				InputViewport.Rect.Min.X, InputViewport.Rect.Min.Y, InputViewport.Rect.Width(), InputViewport.Rect.Height(),
+				OutputViewport.Rect.Size(),
+				InputViewport.Extent,
+				VertexShader,
+				EDRF_UseTriangleOptimization);
+		}
 	});
 
 	return MoveTemp(Output);
diff --git a/Engine/Source/Runtime/Renderer/Private/PostProcess/TemporalSuperResolution.cpp b/Engine/Source/Runtime/Renderer/Private/PostProcess/TemporalSuperResolution.cpp
--- a/Engine/Source/Runtime/Renderer/Private/PostProcess/TemporalSuperResolution.cpp
+++ b/Engine/Source/Runtime/Renderer/Private/PostProcess/TemporalSuperResolution.cpp
@@ -1067,6 +1067,7 @@ class FTSRUpdateHistoryCS : public FTSRShader
 		SHADER_PARAMETER_RDG_TEXTURE(Texture2D, PrevDistortingDisplacementTexture)
 		SHADER_PARAMETER_RDG_TEXTURE(Texture2D, ResurrectedDistortingDisplacementTexture)
 		SHADER_PARAMETER_RDG_TEXTURE(Texture2D, UndistortingDisplacementTexture)
+		SHADER_PARAMETER(float, DistortionOverscan)
 
 		SHADER_PARAMETER_RDG_TEXTURE_UAV(RWTexture2DArray, HistoryColorOutput)
 		SHADER_PARAMETER_RDG_TEXTURE_UAV(RWTexture2DArray, HistoryMetadataOutput)
@@ -2728,9 +2729,12 @@ FDefaultTemporalUpscaler::FOutputs AddTemporalSuperResolutionPasses(
 		PassParameters->PrevDistortingDisplacementTexture = PrevDistortingDisplacementTexture;
 		PassParameters->ResurrectedDistortingDisplacementTexture = ResurrectedDistortingDisplacementTexture;
 		PassParameters->UndistortingDisplacementTexture = BlackDummy;
+		PassParameters->DistortionOverscan = 1.0f;
+		
 		if (bLensDistortion && PassInputs.LensDistortionLUT.IsEnabled())
 		{
 			PassParameters->UndistortingDisplacementTexture = PassInputs.LensDistortionLUT.UndistortingDisplacementTexture;
+			PassParameters->DistortionOverscan = PassInputs.LensDistortionLUT.DistortionOverscan;
 		}
 
 		{
