Fix minor shader warning

#rnx

diff --git a/Engine/Plugins/VirtualProduction/CameraCalibrationCore/Shaders/Private/FillSTDisplacementMap.usf b/Engine/Plugins/VirtualProduction/CameraCalibrationCore/Shaders/Private/FillSTDisplacementMap.usf
--- a/Engine/Plugins/VirtualProduction/CameraCalibrationCore/Shaders/Private/FillSTDisplacementMap.usf
+++ b/Engine/Plugins/VirtualProduction/CameraCalibrationCore/Shaders/Private/FillSTDisplacementMap.usf
@@ -28,7 +28,7 @@ void MainVS(
 	// Compute UV from vertex ID and rectangle layout
 	float2 UnoverscannedUV = float2(VertexId % (GridDimensions.x + 1), VertexId / (GridDimensions.x + 1)) / float2(GridDimensions) + 0.5 * PixelToUV;
 	float2 ViewportUV = ToOverscannedSpace(UnoverscannedUV);
-	float2 DisplacedUV = ToOverscannedSpace(UnoverscannedUV + UndisplacementMap.SampleLevel(UndisplacementMapSampler, UnoverscannedUV, 0));
+	float2 DisplacedUV = ToOverscannedSpace(UnoverscannedUV + UndisplacementMap.SampleLevel(UndisplacementMapSampler, UnoverscannedUV, 0).xy);
 	
 	// Write the original position in original and overscanned space to the texture coordinate
 	OutTexCoord.xy = ViewportUV;
