Camera Calibration: Fix Lens File CDO allocating render targets

diff --git a/Engine/Plugins/VirtualProduction/CameraCalibrationCore/Source/CameraCalibrationCore/Private/LensFile.cpp b/Engine/Plugins/VirtualProduction/CameraCalibrationCore/Source/CameraCalibrationCore/Private/LensFile.cpp
--- a/Engine/Plugins/VirtualProduction/CameraCalibrationCore/Source/CameraCalibrationCore/Private/LensFile.cpp
+++ b/Engine/Plugins/VirtualProduction/CameraCalibrationCore/Source/CameraCalibrationCore/Private/LensFile.cpp
@@ -1063,10 +1063,13 @@ FBaseLensTable* ULensFile::GetDataTable(ELensDataCategory InDataCategory)
 void ULensFile::PostInitProperties()
 {
 	Super::PostInitProperties();
-	
-	const FIntPoint DisplacementMapResolution = GetDefault<UCameraCalibrationSettings>()->GetDisplacementMapResolution();
-	CreateIntermediateDisplacementMaps(DisplacementMapResolution);
 
+	if (!HasAnyFlags(RF_ClassDefaultObject))
+	{
+		const FIntPoint DisplacementMapResolution = GetDefault<UCameraCalibrationSettings>()->GetDisplacementMapResolution();
+		CreateIntermediateDisplacementMaps(DisplacementMapResolution);
+	}
+	
 	// Set a Lens file reference to all tables
 	DistortionTable.LensFile =
 		FocalLengthTable.LensFile = ImageCenterTable.LensFile =
