Fixup API macros for merged modules, PURE_VIRTUAL does not need API export

diff --git a/Engine/Plugins/Runtime/ReplicationGraph/Source/Public/ReplicationGraph.h b/Engine/Plugins/Runtime/ReplicationGraph/Source/Public/ReplicationGraph.h
--- a/Engine/Plugins/Runtime/ReplicationGraph/Source/Public/ReplicationGraph.h
+++ b/Engine/Plugins/Runtime/ReplicationGraph/Source/Public/ReplicationGraph.h
@@ -77,17 +77,17 @@ public:
 	UE_API virtual void Serialize(FArchive& Ar) override;
 
 	/** Called when a network actor is spawned or an actor changes replication status */
-	UE_API virtual void NotifyAddNetworkActor(const FNewReplicatedActorInfo& Actor ) PURE_VIRTUAL(UReplicationGraphNode::NotifyAddNetworkActor, );
+	virtual void NotifyAddNetworkActor(const FNewReplicatedActorInfo& Actor ) PURE_VIRTUAL(UReplicationGraphNode::NotifyAddNetworkActor, );
 	
 	/** Called when a networked actor is being destroyed or no longer wants to replicate */
-	UE_API virtual bool NotifyRemoveNetworkActor(const FNewReplicatedActorInfo& Actor, bool bWarnIfNotFound=true) PURE_VIRTUAL(UReplicationGraphNode::NotifyRemoveNetworkActor, return false; );
+	virtual bool NotifyRemoveNetworkActor(const FNewReplicatedActorInfo& Actor, bool bWarnIfNotFound=true) PURE_VIRTUAL(UReplicationGraphNode::NotifyRemoveNetworkActor, return false; );
 
 	/**
 	 * If Net.RepGraph.HandleDynamicActorRename is enabled, this is called when a dynamic actor is changing its outer.
 	 * If the node maintains level information for actors, and if replicated actors may be renamed, the implementation needs to update its level info.
 	 * Returns true if the actor was found & updated.
 	 */
-	UE_API virtual bool NotifyActorRenamed(const FRenamedReplicatedActorInfo& Actor, bool bWarnIfNotFound=true) PURE_VIRTUAL(UReplicationGraphNode::NotifyActorRenamed, return false; );
+	virtual bool NotifyActorRenamed(const FRenamedReplicatedActorInfo& Actor, bool bWarnIfNotFound=true) PURE_VIRTUAL(UReplicationGraphNode::NotifyActorRenamed, return false; );
 
 	/** Called when world changes or when all subclasses should dump any persistent data/lists about replicated actors here. (The new/next world will be set before this is called) */
 	UE_API virtual void NotifyResetAllNetworkActors();
@@ -95,7 +95,7 @@ public:
 	/** Mark the node and all its children PendingKill */
 	UE_API virtual void TearDown();
 	
-	UE_API virtual void GatherActorListsForConnection(const FConnectionGatherActorListParameters& Params) PURE_VIRTUAL(UReplicationGraphNode::GatherActorListsForConnection, );
+	virtual void GatherActorListsForConnection(const FConnectionGatherActorListParameters& Params) PURE_VIRTUAL(UReplicationGraphNode::GatherActorListsForConnection, );
 
 	/** Called once per frame prior to replication ONLY on root nodes (nodes created via UReplicationGraph::CreateNode) has RequiresPrepareForReplicationCall=true */
 	virtual void PrepareForReplication() { };
diff --git a/Engine/Plugins/Runtime/SoundCueTemplates/Source/SoundCueTemplates/Public/SoundCueTemplate.h b/Engine/Plugins/Runtime/SoundCueTemplates/Source/SoundCueTemplates/Public/SoundCueTemplate.h
--- a/Engine/Plugins/Runtime/SoundCueTemplates/Source/SoundCueTemplates/Public/SoundCueTemplate.h
+++ b/Engine/Plugins/Runtime/SoundCueTemplates/Source/SoundCueTemplates/Public/SoundCueTemplate.h
@@ -44,7 +44,7 @@ protected:
 	 * graph of provided SoundCue.  Provided SoundCue may be a reference to this SoundCue
 	 * or one to be copied from this template's graph-building schema.
 	 */
-	SOUNDCUETEMPLATES_API virtual void OnRebuildGraph(USoundCue& SoundCue) const PURE_VIRTUAL(, );
+	virtual void OnRebuildGraph(USoundCue& SoundCue) const PURE_VIRTUAL(, );
 
 	/**
 	 * Node width offset between nodes in the SoundCueEditor.
diff --git a/Engine/Plugins/Runtime/ZoneGraph/Source/ZoneGraphDebug/Public/ZoneGraphTestingActor.h b/Engine/Plugins/Runtime/ZoneGraph/Source/ZoneGraphDebug/Public/ZoneGraphTestingActor.h
--- a/Engine/Plugins/Runtime/ZoneGraph/Source/ZoneGraphDebug/Public/ZoneGraphTestingActor.h
+++ b/Engine/Plugins/Runtime/ZoneGraph/Source/ZoneGraphDebug/Public/ZoneGraphTestingActor.h
@@ -23,7 +23,7 @@ class UZoneLaneTest : public UObject
 	GENERATED_BODY()
 
 public:
-	UE_API virtual void OnLaneLocationUpdated(const FZoneGraphLaneLocation& PrevLaneLocation, const FZoneGraphLaneLocation& NextLaneLocation) PURE_VIRTUAL(UZoneLaneTest::OnLaneLocationUpdated, );
+	virtual void OnLaneLocationUpdated(const FZoneGraphLaneLocation& PrevLaneLocation, const FZoneGraphLaneLocation& NextLaneLocation) PURE_VIRTUAL(UZoneLaneTest::OnLaneLocationUpdated, );
 	virtual void Draw(FPrimitiveDrawInterface* PDI) const {};
 
 	void SetOwner(UZoneGraphTestingComponent* Owner) { OwnerComponent = Owner; OnOwnerSet(); }
diff --git a/Engine/Plugins/VirtualProduction/CameraCalibrationCore/Source/CameraCalibrationCore/Public/LensDistortionModelHandlerBase.h b/Engine/Plugins/VirtualProduction/CameraCalibrationCore/Source/CameraCalibrationCore/Public/LensDistortionModelHandlerBase.h
--- a/Engine/Plugins/VirtualProduction/CameraCalibrationCore/Source/CameraCalibrationCore/Public/LensDistortionModelHandlerBase.h
+++ b/Engine/Plugins/VirtualProduction/CameraCalibrationCore/Source/CameraCalibrationCore/Public/LensDistortionModelHandlerBase.h
@@ -129,22 +129,22 @@ public:
 
 protected:
 	/** Initialize the handler. Derived classes must set the LensModelClass that they support, if not already set */
-	UE_API virtual void InitializeHandler() PURE_VIRTUAL(ULensDistortionModelHandlerBase::InitializeHandler);
+	virtual void InitializeHandler() PURE_VIRTUAL(ULensDistortionModelHandlerBase::InitializeHandler);
 
 	/** Use the current distortion state to compute the distortion position of an input UV coordinate */
-	UE_API virtual FVector2D ComputeDistortedUV(const FVector2D& InScreenUV) const PURE_VIRTUAL(ULensDistortionModelHandlerBase::ComputeDistortedUV, return FVector2D::ZeroVector;);
+	virtual FVector2D ComputeDistortedUV(const FVector2D& InScreenUV) const PURE_VIRTUAL(ULensDistortionModelHandlerBase::ComputeDistortedUV, return FVector2D::ZeroVector;);
 
 	/** Use the current distortion state to compute the undistortion position of an input UV coordinate */
-	UE_API virtual FVector2D ComputeUndistortedUV(const FVector2D& InScreenUV) const PURE_VIRTUAL(ULensDistortionModelHandlerBase::ComputeUndistortedUV, return FVector2D::ZeroVector;);
+	virtual FVector2D ComputeUndistortedUV(const FVector2D& InScreenUV) const PURE_VIRTUAL(ULensDistortionModelHandlerBase::ComputeUndistortedUV, return FVector2D::ZeroVector;);
 	
 	/** Create the distortion MIDs */
-	UE_API virtual void InitDistortionMaterials() PURE_VIRTUAL(ULensDistortionModelHandlerBase::InitDistortionMaterials);
+	virtual void InitDistortionMaterials() PURE_VIRTUAL(ULensDistortionModelHandlerBase::InitDistortionMaterials);
 
 	/** Set the material parameters for the displacement map and distortion post-process materials */
-	UE_API virtual void UpdateMaterialParameters() PURE_VIRTUAL(ULensDistortionModelHandlerBase::UpdateMaterialParameters);
+	virtual void UpdateMaterialParameters() PURE_VIRTUAL(ULensDistortionModelHandlerBase::UpdateMaterialParameters);
 
 	/** Convert the generic distortion parameter array into the specific structure of parameters used by the supported lens model */
-	UE_API virtual void InterpretDistortionParameters() PURE_VIRTUAL(ULensDistortionModelHandlerBase::InterpretDistortionParameters);
+	virtual void InterpretDistortionParameters() PURE_VIRTUAL(ULensDistortionModelHandlerBase::InterpretDistortionParameters);
 
 private:
 	/** Tests whether the input material's shaders have been compiled, indicating the material can be used for drawing */
diff --git a/Engine/Plugins/VirtualProduction/CameraCalibrationCore/Source/CameraCalibrationCore/Public/Models/LensModel.h b/Engine/Plugins/VirtualProduction/CameraCalibrationCore/Source/CameraCalibrationCore/Public/Models/LensModel.h
--- a/Engine/Plugins/VirtualProduction/CameraCalibrationCore/Source/CameraCalibrationCore/Public/Models/LensModel.h
+++ b/Engine/Plugins/VirtualProduction/CameraCalibrationCore/Source/CameraCalibrationCore/Public/Models/LensModel.h
@@ -24,13 +24,13 @@ class ULensModel : public UObject
 
 public:
 	/** Get the lens model name */
-	UE_API virtual FName GetModelName() const PURE_VIRTUAL(ULensModel::GetModelName, return FName(""););
+	virtual FName GetModelName() const PURE_VIRTUAL(ULensModel::GetModelName, return FName(""););
 
 	/** Get the lens model short name */
-	UE_API virtual FName GetShortModelName() const PURE_VIRTUAL(ULensModel::GetShortModelName, return FName(""););
+	virtual FName GetShortModelName() const PURE_VIRTUAL(ULensModel::GetShortModelName, return FName(""););
 
 	/** Get the struct of distortion parameters supported by this model */
-	UE_API virtual UScriptStruct* GetParameterStruct() const PURE_VIRTUAL(ULensModel::GetParameterStruct, return nullptr;);
+	virtual UScriptStruct* GetParameterStruct() const PURE_VIRTUAL(ULensModel::GetParameterStruct, return nullptr;);
 
 	/** Returns an array of floats representing the default values of each property in the parameter struct */
 	UE_API void GetDefaultParameterArray(TArray<float>& OutParameterValues) const;
diff --git a/Engine/Plugins/VirtualProduction/CameraCalibrationCore/Source/CameraCalibrationCore/Public/Tables/BaseLensTable.h b/Engine/Plugins/VirtualProduction/CameraCalibrationCore/Source/CameraCalibrationCore/Public/Tables/BaseLensTable.h
--- a/Engine/Plugins/VirtualProduction/CameraCalibrationCore/Source/CameraCalibrationCore/Public/Tables/BaseLensTable.h
+++ b/Engine/Plugins/VirtualProduction/CameraCalibrationCore/Source/CameraCalibrationCore/Public/Tables/BaseLensTable.h
@@ -44,13 +44,13 @@ public:
 	virtual ~FBaseFocusPoint() = default;
 
 	/** Returns focus value for this Focus Point */
-	UE_API virtual float GetFocus() const PURE_VIRTUAL(FBaseLensTable::GetFocus, return 0.f;);
+	virtual float GetFocus() const PURE_VIRTUAL(FBaseLensTable::GetFocus, return 0.f;);
 
 	/** Returns number of zoom points */
-	UE_API virtual int32 GetNumPoints() const PURE_VIRTUAL(FBaseLensTable::GetNumPoints, return 0;);
+	virtual int32 GetNumPoints() const PURE_VIRTUAL(FBaseLensTable::GetNumPoints, return 0;);
 
 	/** Returns zoom value for a given index */
-	UE_API virtual float GetZoom(int32 Index) const PURE_VIRTUAL(FBaseLensTable::GetZoom, return 0.f;);
+	virtual float GetZoom(int32 Index) const PURE_VIRTUAL(FBaseLensTable::GetZoom, return 0.f;);
 };
 
 template<>
@@ -104,14 +104,14 @@ struct FBaseLensTable
 protected:
 
 	/** Returns the map of linked categories  */
-	UE_API virtual TMap<ELensDataCategory, FLinkPointMetadata> GetLinkedCategories() const PURE_VIRTUAL(FBaseLensTable::GetLinkedCategories, return TMap<ELensDataCategory, FLinkPointMetadata>() ;);
+	virtual TMap<ELensDataCategory, FLinkPointMetadata> GetLinkedCategories() const PURE_VIRTUAL(FBaseLensTable::GetLinkedCategories, return TMap<ELensDataCategory, FLinkPointMetadata>() ;);
 
 	/**
 	 * Whether the focus point exists
 	 * @param InFocus focus value to check
 	 * @return true if the point exists
 	 */
-	UE_API virtual bool DoesFocusPointExists(float InFocus, float InputTolerance = KINDA_SMALL_NUMBER) const PURE_VIRTUAL(FBaseLensTable::DoesFocusPointExists, return false; );
+	virtual bool DoesFocusPointExists(float InFocus, float InputTolerance = KINDA_SMALL_NUMBER) const PURE_VIRTUAL(FBaseLensTable::DoesFocusPointExists, return false; );
 
 	/**
 	 * Whether the zoom point exists
@@ -120,7 +120,7 @@ protected:
 	 * @param InputTolerance Maximum allowed difference for considering them as 'nearly equal'
 	 * @return true if zoom point exists
 	 */
-	UE_API virtual bool DoesZoomPointExists(float InFocus, float InZoom, float InputTolerance = KINDA_SMALL_NUMBER) const PURE_VIRTUAL(FBaseLensTable::DoesZoomPointExists, return false; );
+	virtual bool DoesZoomPointExists(float InFocus, float InZoom, float InputTolerance = KINDA_SMALL_NUMBER) const PURE_VIRTUAL(FBaseLensTable::DoesZoomPointExists, return false; );
 
 	/** Copies the specified keys from the source curve to the destination curve */
 	UE_API void CopyCurveKeys(const FRichCurve& InSourceCurve, FRichCurve& InDestCurve, TArrayView<const FKeyHandle> InKeys);
@@ -135,25 +135,25 @@ public:
 	 * Loop through all Focus Points
 	 * @param InCallback Callback with Focus point reference
 	 */
-	UE_API virtual void ForEachPoint(FFocusPointCallback InCallback) const PURE_VIRTUAL(FBaseLensTable::ForEachPoint );
+	virtual void ForEachPoint(FFocusPointCallback InCallback) const PURE_VIRTUAL(FBaseLensTable::ForEachPoint );
 
 	/** Get number of Focus points for this data table */
-	UE_API virtual int32 GetFocusPointNum() const PURE_VIRTUAL(FBaseLensTable::GetFocusPointNum, return INDEX_NONE; );
+	virtual int32 GetFocusPointNum() const PURE_VIRTUAL(FBaseLensTable::GetFocusPointNum, return INDEX_NONE; );
 
 	/** Get total number of Zoom points for all Focus points of this data table */
-	UE_API virtual int32 GetTotalPointNum() const PURE_VIRTUAL(FBaseLensTable::GetTotalPointNum, return INDEX_NONE; );
+	virtual int32 GetTotalPointNum() const PURE_VIRTUAL(FBaseLensTable::GetTotalPointNum, return INDEX_NONE; );
 
 	/** Get the base focus point by given index */
-	UE_API virtual const FBaseFocusPoint* GetBaseFocusPoint(int32 InIndex) const PURE_VIRTUAL(FBaseLensTable::GetBaseFocusPoint, return nullptr; );
+	virtual const FBaseFocusPoint* GetBaseFocusPoint(int32 InIndex) const PURE_VIRTUAL(FBaseLensTable::GetBaseFocusPoint, return nullptr; );
 
 	/** Get Struct class of this Data Table */
-	UE_API virtual UScriptStruct* GetScriptStruct() const PURE_VIRTUAL(FBaseLensTable::GetFocusPointNum, return nullptr; );
+	virtual UScriptStruct* GetScriptStruct() const PURE_VIRTUAL(FBaseLensTable::GetFocusPointNum, return nullptr; );
 
 	/** 
 	* Fills OutCurve with all points contained in the given focus 
 	* Returns false if FocusIdentifier is not found or ParameterIndex isn't valid
 	*/
-	UE_API virtual bool BuildParameterCurveAtFocus(float InFocus, int32 ParameterIndex, FRichCurve& OutCurve) const PURE_VIRTUAL(FBaseLensTable::BuildParameterCurve, return false; ); 
+	virtual bool BuildParameterCurveAtFocus(float InFocus, int32 ParameterIndex, FRichCurve& OutCurve) const PURE_VIRTUAL(FBaseLensTable::BuildParameterCurve, return false; ); 
 
 	/**
 	 * Fills OutCurve with points across all focuses that have the given zoom
@@ -162,7 +162,7 @@ public:
 	 * @param OutCurve The curve to fill
 	 * @return true if the curve was built, otherwise false
 	 */
-	UE_API virtual bool BuildParameterCurveAtZoom(float InZoom, int32 ParameterIndex, FRichCurve& OutCurve) const PURE_VIRTUAL(FBaseLensTable::BuildFocusParameterCurve, return false; );
+	virtual bool BuildParameterCurveAtZoom(float InZoom, int32 ParameterIndex, FRichCurve& OutCurve) const PURE_VIRTUAL(FBaseLensTable::BuildFocusParameterCurve, return false; );
 
 	/**
 	 * Updates the keys of all zoom points at the specified focus to match the corresponding keys in the specified curve
@@ -171,7 +171,7 @@ public:
 	 * @param InSourceCurve The curve to copy the keys from
 	 * @param InKeys The keys to copy from the source curve into the zoom parameter curve
 	 */
-	UE_API virtual void SetParameterCurveKeysAtFocus(float InFocus, int32 InParameterIndex, const FRichCurve& InSourceCurve, TArrayView<const FKeyHandle> InKeys) PURE_VIRTUAL(FBaseLensTable::SetZoomParameterCurveKeys);
+	virtual void SetParameterCurveKeysAtFocus(float InFocus, int32 InParameterIndex, const FRichCurve& InSourceCurve, TArrayView<const FKeyHandle> InKeys) PURE_VIRTUAL(FBaseLensTable::SetZoomParameterCurveKeys);
 
 	/**
 	 * Updates the keys of all focus points at the specified zoom to match the corresponding keys in the specified curve
@@ -180,21 +180,21 @@ public:
 	 * @param InSourceCurve The curve to copy the keys from
 	 * @param InKeys The keys to copy from the source curve into the focus parameter curve
 	 */
-	UE_API virtual void SetParameterCurveKeysAtZoom(float InZoom, int32 InParameterIndex, const FRichCurve& InSourceCurve, TArrayView<const FKeyHandle> InKeys) PURE_VIRTUAL(FBaseLensTable::SetFocusParameterCurveKeys);
+	virtual void SetParameterCurveKeysAtZoom(float InZoom, int32 InParameterIndex, const FRichCurve& InSourceCurve, TArrayView<const FKeyHandle> InKeys) PURE_VIRTUAL(FBaseLensTable::SetFocusParameterCurveKeys);
 
 	/**
 	 * Gets whether the positions of the table's curve keys can be edited or not
 	 * @param InParameterIndex The index of the parameter of the curve in question
 	 * @return true if the curve's key positions can be edited, false otherwise
 	 */
-	UE_API virtual bool CanEditCurveKeyPositions(int32 InParameterIndex) const PURE_VIRTUAL(FBaseLensTable::CanEditCurveKeyPositions, return false; );
+	virtual bool CanEditCurveKeyPositions(int32 InParameterIndex) const PURE_VIRTUAL(FBaseLensTable::CanEditCurveKeyPositions, return false; );
 
 	/**
 	 * Gets whether the attributes of the table's curve keys can be edited or not
 	 * @param InParameterIndex The index of the parameter of the curve in question
 	 * @return true if the curve's key attributes can be edited, false otherwise
 	 */
-	UE_API virtual bool CanEditCurveKeyAttributes(int32 InParameterIndex) const PURE_VIRTUAL(FBaseLensTable::CanEditCurveKeyAttributes, return false; );
+	virtual bool CanEditCurveKeyAttributes(int32 InParameterIndex) const PURE_VIRTUAL(FBaseLensTable::CanEditCurveKeyAttributes, return false; );
 
 	/**
 	 * Gets the range of allowed values for the curve keys at the specified parameter index
diff --git a/Engine/Plugins/WorldMetrics/Source/WorldMetricsCore/Public/WorldMetricInterface.h b/Engine/Plugins/WorldMetrics/Source/WorldMetricsCore/Public/WorldMetricInterface.h
--- a/Engine/Plugins/WorldMetrics/Source/WorldMetricsCore/Public/WorldMetricInterface.h
+++ b/Engine/Plugins/WorldMetrics/Source/WorldMetricsCore/Public/WorldMetricInterface.h
@@ -21,7 +21,7 @@ public:
 	 * Returns the amount of memory allocated by this class, not including sizeof(*this).
 	 * @return the result size in bytes.
 	 */
-	[[nodiscard]] WORLDMETRICSCORE_API virtual SIZE_T GetAllocatedSize() const
+	[[nodiscard]] virtual SIZE_T GetAllocatedSize() const
 		PURE_VIRTUAL(GetAllocatedSize, return 0;);
 
 	/*
diff --git a/Engine/Plugins/WorldMetrics/Source/WorldMetricsCore/Public/WorldMetricsExtension.h b/Engine/Plugins/WorldMetrics/Source/WorldMetricsCore/Public/WorldMetricsExtension.h
--- a/Engine/Plugins/WorldMetrics/Source/WorldMetricsCore/Public/WorldMetricsExtension.h
+++ b/Engine/Plugins/WorldMetrics/Source/WorldMetricsCore/Public/WorldMetricsExtension.h
@@ -26,7 +26,7 @@ public:
 	 * Returns the amount of memory allocated by this class, not including sizeof(*this).
 	 * @return the result size in bytes.
 	 */
-	[[nodiscard]] WORLDMETRICSCORE_API virtual SIZE_T GetAllocatedSize() const
+	[[nodiscard]] virtual SIZE_T GetAllocatedSize() const
 		PURE_VIRTUAL(GetAllocatedSize, return 0;);
 
 	/*
diff --git a/Engine/Source/Developer/FunctionalTesting/Classes/FunctionalAITest.h b/Engine/Source/Developer/FunctionalTesting/Classes/FunctionalAITest.h
--- a/Engine/Source/Developer/FunctionalTesting/Classes/FunctionalAITest.h
+++ b/Engine/Source/Developer/FunctionalTesting/Classes/FunctionalAITest.h
@@ -58,7 +58,7 @@ struct FAITestSpawnInfoBase
 
 	virtual bool IsValid() const { return SpawnLocation != NULL; }
 
-	UE_API virtual bool Spawn(AFunctionalAITestBase* AITest) const PURE_VIRTUAL(, return false;);
+	virtual bool Spawn(AFunctionalAITestBase* AITest) const PURE_VIRTUAL(, return false;);
 };
 
 template<>
@@ -245,19 +245,19 @@ public:
 	UE_API AFunctionalAITestBase(const FObjectInitializer& ObjectInitializer = FObjectInitializer::Get());
 
 	/** Iterate through the list of spawn sets and execute Predicate on each in a const-correct way. */
-	UE_API virtual void ForEachSpawnSet(TFunctionRef<void(const FAITestSpawnSetBase&)> Predicate) const PURE_VIRTUAL(, );
+	virtual void ForEachSpawnSet(TFunctionRef<void(const FAITestSpawnSetBase&)> Predicate) const PURE_VIRTUAL(, );
 
 	/** Iterate through the list of spawn sets and execute Predicate on each. */
-	UE_API virtual void ForEachSpawnSet(TFunctionRef<void(FAITestSpawnSetBase&)> Predicate) PURE_VIRTUAL(, );
+	virtual void ForEachSpawnSet(TFunctionRef<void(FAITestSpawnSetBase&)> Predicate) PURE_VIRTUAL(, );
 
 	/** Iterate through the list of spawn sets and remove the spawn set if Predicate returns true. */
-	UE_API virtual void RemoveSpawnSetIfPredicate(TFunctionRef<bool(FAITestSpawnSetBase&)> Predicate) PURE_VIRTUAL(, );
+	virtual void RemoveSpawnSetIfPredicate(TFunctionRef<bool(FAITestSpawnSetBase&)> Predicate) PURE_VIRTUAL(, );
 
 	/** Return the SpawnSet at this index of the SpawnSets array. Const-correct version. */
-	UE_API virtual const FAITestSpawnSetBase* GetSpawnSet(const int32 SpawnSetIndex) const PURE_VIRTUAL(, return nullptr;);
+	virtual const FAITestSpawnSetBase* GetSpawnSet(const int32 SpawnSetIndex) const PURE_VIRTUAL(, return nullptr;);
 
 	/** Return the SpawnSet at this index of the SpawnSets array. */
-	UE_API virtual FAITestSpawnSetBase* GetSpawnSet(const int32 SpawnSetIndex) PURE_VIRTUAL(,return nullptr;);
+	virtual FAITestSpawnSetBase* GetSpawnSet(const int32 SpawnSetIndex) PURE_VIRTUAL(,return nullptr;);
 	
 	/** Return the SpawnInfo at SpawnInfoIndexof the SpawnSet at SpawnSetIndex. Const-correct version. */
 	UE_API const FAITestSpawnInfoBase* GetSpawnInfo(const int32 SpawnSetIndex, const int32 SpawnInfoIndex) const;
@@ -266,7 +266,7 @@ public:
 	UE_API FAITestSpawnInfoBase* GetSpawnInfo(const int32 SpawnSetIndex, const int32 SpawnInfoIndex);
 
 	/** Return whether the index is valid in the SpawnSets array. */
-	UE_API virtual bool IsValidSpawnSetIndex(const int32 Index) const PURE_VIRTUAL(, return false;);
+	virtual bool IsValidSpawnSetIndex(const int32 Index) const PURE_VIRTUAL(, return false;);
 
 	/** 
 	* Spawn this AI at this SpawnInfoIndex of the SpawnSetIndex spawn set. 
diff --git a/Engine/Source/Runtime/AIModule/Classes/BehaviorTree/BlackboardAssetProvider.h b/Engine/Source/Runtime/AIModule/Classes/BehaviorTree/BlackboardAssetProvider.h
--- a/Engine/Source/Runtime/AIModule/Classes/BehaviorTree/BlackboardAssetProvider.h
+++ b/Engine/Source/Runtime/AIModule/Classes/BehaviorTree/BlackboardAssetProvider.h
@@ -29,7 +29,7 @@ class IBlackboardAssetProvider
 #endif
 	/** Returns BlackboardData referenced by the owner object. */
 	UFUNCTION(BlueprintCallable, Category = GameplayTags)
-	AIMODULE_API virtual UBlackboardData* GetBlackboardAsset() const PURE_VIRTUAL(IBlackboardAssetProvider::GetBlackboardAsset, return nullptr; );
+	virtual UBlackboardData* GetBlackboardAsset() const PURE_VIRTUAL(IBlackboardAssetProvider::GetBlackboardAsset, return nullptr; );
 };
 
 
diff --git a/Engine/Source/Runtime/AIModule/Classes/Perception/AISenseEvent.h b/Engine/Source/Runtime/AIModule/Classes/Perception/AISenseEvent.h
--- a/Engine/Source/Runtime/AIModule/Classes/Perception/AISenseEvent.h
+++ b/Engine/Source/Runtime/AIModule/Classes/Perception/AISenseEvent.h
@@ -15,7 +15,7 @@ class UAISenseEvent : public UObject
 	GENERATED_BODY()
 
 public:
-	AIMODULE_API virtual FAISenseID GetSenseID() const PURE_VIRTUAL(UAISenseEvent::GetSenseID, return FAISenseID::InvalidID(););
+	virtual FAISenseID GetSenseID() const PURE_VIRTUAL(UAISenseEvent::GetSenseID, return FAISenseID::InvalidID(););
 
 #if ENABLE_VISUAL_LOG
 	virtual void DrawToVLog(UObject& LogOwner) const {}
diff --git a/Engine/Source/Runtime/AnimGraphRuntime/Public/AnimNodes/AnimNode_BlendSpacePlayer.h b/Engine/Source/Runtime/AnimGraphRuntime/Public/AnimNodes/AnimNode_BlendSpacePlayer.h
--- a/Engine/Source/Runtime/AnimGraphRuntime/Public/AnimNodes/AnimNode_BlendSpacePlayer.h
+++ b/Engine/Source/Runtime/AnimGraphRuntime/Public/AnimNodes/AnimNode_BlendSpacePlayer.h
@@ -58,26 +58,26 @@ public:
 public:
 
 	// Get the blendspace asset to play
-	ANIMGRAPHRUNTIME_API virtual UBlendSpace* GetBlendSpace() const PURE_VIRTUAL(FAnimNode_BlendSpacePlayerBase::GetBlendSpace, return nullptr;);
+	virtual UBlendSpace* GetBlendSpace() const PURE_VIRTUAL(FAnimNode_BlendSpacePlayerBase::GetBlendSpace, return nullptr;);
 
 	// Get the coordinates that are currently being sampled by the blendspace
-	ANIMGRAPHRUNTIME_API virtual FVector GetPosition() const PURE_VIRTUAL(FAnimNode_BlendSpacePlayerBase::GetPosition, return FVector::Zero(););
+	virtual FVector GetPosition() const PURE_VIRTUAL(FAnimNode_BlendSpacePlayerBase::GetPosition, return FVector::Zero(););
 
 	// The start position in [0, 1] to use when initializing. When looping, play will still jump back to the beginning when reaching the end.
-	ANIMGRAPHRUNTIME_API virtual float GetStartPosition() const PURE_VIRTUAL(FAnimNode_BlendSpacePlayerBase::GetStartPosition, return 0.0f;);
+	virtual float GetStartPosition() const PURE_VIRTUAL(FAnimNode_BlendSpacePlayerBase::GetStartPosition, return 0.0f;);
 
 	// Get the play rate multiplier. Can be negative, which will cause the animation to play in reverse.
-	ANIMGRAPHRUNTIME_API virtual float GetPlayRate() const PURE_VIRTUAL(FAnimNode_BlendSpacePlayerBase::GetPlayRate, return 1.0f;);
+	virtual float GetPlayRate() const PURE_VIRTUAL(FAnimNode_BlendSpacePlayerBase::GetPlayRate, return 1.0f;);
 
 	// Should the animation loop back to the start when it reaches the end?
 	UE_DEPRECATED(5.3, "Please use IsLooping instead.")
 	virtual bool GetLoop() const final { return IsLooping(); }
 
 	// Get whether we should reset the current play time when the blend space changes
-	ANIMGRAPHRUNTIME_API virtual bool ShouldResetPlayTimeWhenBlendSpaceChanges() const PURE_VIRTUAL(FAnimNode_BlendSpacePlayerBase::ShouldResetPlayTimeWhenBlendSpaceChanges, return true;);
+	virtual bool ShouldResetPlayTimeWhenBlendSpaceChanges() const PURE_VIRTUAL(FAnimNode_BlendSpacePlayerBase::ShouldResetPlayTimeWhenBlendSpaceChanges, return true;);
 
 	// Set whether we should reset the current play time when the blend space changes
-	ANIMGRAPHRUNTIME_API virtual bool SetResetPlayTimeWhenBlendSpaceChanges(bool bReset) PURE_VIRTUAL(FAnimNode_BlendSpacePlayerBase::SetResetPlayTimeWhenBlendSpaceChanges, return false;);
+	virtual bool SetResetPlayTimeWhenBlendSpaceChanges(bool bReset) PURE_VIRTUAL(FAnimNode_BlendSpacePlayerBase::SetResetPlayTimeWhenBlendSpaceChanges, return false;);
 
 	// An evaluator will be setting the play rate to zero and setting the time explicitly. ShouldTeleportToTime indicates whether we should jump to that time, or move to it playing out root motion and events etc.
 	virtual bool ShouldTeleportToTime() const { return false; }
@@ -86,16 +86,16 @@ public:
 	virtual bool IsEvaluator() const { return false; }
 
 	// Set the blendspace asset to play
-	ANIMGRAPHRUNTIME_API virtual bool SetBlendSpace(UBlendSpace* InBlendSpace) PURE_VIRTUAL(FAnimNode_BlendSpacePlayerBase::SetBlendSpace, return false;);
+	virtual bool SetBlendSpace(UBlendSpace* InBlendSpace) PURE_VIRTUAL(FAnimNode_BlendSpacePlayerBase::SetBlendSpace, return false;);
 
 	// Set the coordinates that are currently being sampled by the blendspace
-	ANIMGRAPHRUNTIME_API virtual bool SetPosition(FVector InPosition) PURE_VIRTUAL(FAnimNode_BlendSpacePlayerBase::SetPosition, return false;);
+	virtual bool SetPosition(FVector InPosition) PURE_VIRTUAL(FAnimNode_BlendSpacePlayerBase::SetPosition, return false;);
 
 	// Set the play rate multiplier. Can be negative, which will cause the animation to play in reverse.
-	ANIMGRAPHRUNTIME_API virtual bool SetPlayRate(float InPlayRate) PURE_VIRTUAL(FAnimNode_BlendSpacePlayerBase::SetPlayRate, return false;);
+	virtual bool SetPlayRate(float InPlayRate) PURE_VIRTUAL(FAnimNode_BlendSpacePlayerBase::SetPlayRate, return false;);
 
 	// Set if the animation should loop back to the start when it reaches the end?
-	ANIMGRAPHRUNTIME_API virtual bool SetLoop(bool bInLoop) PURE_VIRTUAL(FAnimNode_BlendSpacePlayerBase::SetLoop, return false;);
+	virtual bool SetLoop(bool bInLoop) PURE_VIRTUAL(FAnimNode_BlendSpacePlayerBase::SetLoop, return false;);
 
 protected:
 	ANIMGRAPHRUNTIME_API void UpdateInternal(const FAnimationUpdateContext& Context);
diff --git a/Engine/Source/Runtime/AudioAnalyzer/Classes/AudioAnalyzer.h b/Engine/Source/Runtime/AudioAnalyzer/Classes/AudioAnalyzer.h
--- a/Engine/Source/Runtime/AudioAnalyzer/Classes/AudioAnalyzer.h
+++ b/Engine/Source/Runtime/AudioAnalyzer/Classes/AudioAnalyzer.h
@@ -140,7 +140,7 @@ protected:
 	}
 
 	/* Subclasses must override this method in order to inform this object which AnalyzerFactory to use for analysis */
-	AUDIOANALYZER_API virtual FName GetAnalyzerFactoryName() const PURE_VIRTUAL(UAudioAnalyzer::GetAnalyzerFactoryName, return FName(););
+	virtual FName GetAnalyzerFactoryName() const PURE_VIRTUAL(UAudioAnalyzer::GetAnalyzerFactoryName, return FName(););
 
 	/** How many frames of audio to wait before analyzing the audio. */
 	int32 NumFramesPerBufferToAnalyze = 1024;
diff --git a/Engine/Source/Runtime/AudioAnalyzer/Classes/AudioAnalyzerAsset.h b/Engine/Source/Runtime/AudioAnalyzer/Classes/AudioAnalyzerAsset.h
--- a/Engine/Source/Runtime/AudioAnalyzer/Classes/AudioAnalyzerAsset.h
+++ b/Engine/Source/Runtime/AudioAnalyzer/Classes/AudioAnalyzerAsset.h
@@ -26,12 +26,12 @@ class UAudioAnalyzerAssetBase : public UObject
 		/**
 		 * GetAssetActionName() returns the FText displayed in the editor.
 		 */ 
-		AUDIOANALYZER_API virtual FText GetAssetActionName() const PURE_VIRTUAL(UAudioAnalyzerAsset::GetAssetActionName, return FText(););
+		virtual FText GetAssetActionName() const PURE_VIRTUAL(UAudioAnalyzerAsset::GetAssetActionName, return FText(););
 
 		/**
 		 * GetSupportedClass() returns the class which should be associated with these asset actions.
 		 */ 
-		AUDIOANALYZER_API virtual UClass* GetSupportedClass() const PURE_VIRTUAL(UAudioAnalyzerAsset::GetSupportedClass, return nullptr;);
+		virtual UClass* GetSupportedClass() const PURE_VIRTUAL(UAudioAnalyzerAsset::GetSupportedClass, return nullptr;);
 
 		/**
 		 * GetTypeColor() returns the color used to display asset icons within the editor.
diff --git a/Engine/Source/Runtime/AudioAnalyzer/Classes/AudioAnalyzerNRT.h b/Engine/Source/Runtime/AudioAnalyzer/Classes/AudioAnalyzerNRT.h
--- a/Engine/Source/Runtime/AudioAnalyzer/Classes/AudioAnalyzerNRT.h
+++ b/Engine/Source/Runtime/AudioAnalyzer/Classes/AudioAnalyzerNRT.h
@@ -173,7 +173,7 @@ class UAudioAnalyzerNRT : public UAudioAnalyzerAssetBase
 #endif 
 
 		/* Subclasses must override this method in order to inform this object which AnalyzerNRTFactory to use for analysis */
-		AUDIOANALYZER_API virtual FName GetAnalyzerNRTFactoryName() const PURE_VIRTUAL(UAudioAnalyzerNRT::GetAnalyzerNRTFactoryName, return FName(););
+		virtual FName GetAnalyzerNRTFactoryName() const PURE_VIRTUAL(UAudioAnalyzerNRT::GetAnalyzerNRTFactoryName, return FName(););
 
 	private:
 
diff --git a/Engine/Source/Runtime/AudioExtensions/Public/IAudioEndpoint.h b/Engine/Source/Runtime/AudioExtensions/Public/IAudioEndpoint.h
--- a/Engine/Source/Runtime/AudioExtensions/Public/IAudioEndpoint.h
+++ b/Engine/Source/Runtime/AudioExtensions/Public/IAudioEndpoint.h
@@ -55,7 +55,7 @@ class UAudioEndpointSettingsBase : public UObject
 	GENERATED_BODY()
 
 public:
-	AUDIOEXTENSIONS_API virtual TUniquePtr<IAudioEndpointSettingsProxy> GetProxy() const PURE_VIRTUAL(UAudioEndpointSettingsBase::GetProxy, return nullptr;);
+	virtual TUniquePtr<IAudioEndpointSettingsProxy> GetProxy() const PURE_VIRTUAL(UAudioEndpointSettingsBase::GetProxy, return nullptr;);
 };
 
 //A blank class for when unimplemented endpoint types are returned
diff --git a/Engine/Source/Runtime/AudioExtensions/Public/ISoundfieldEndpoint.h b/Engine/Source/Runtime/AudioExtensions/Public/ISoundfieldEndpoint.h
--- a/Engine/Source/Runtime/AudioExtensions/Public/ISoundfieldEndpoint.h
+++ b/Engine/Source/Runtime/AudioExtensions/Public/ISoundfieldEndpoint.h
@@ -53,7 +53,7 @@ class USoundfieldEndpointSettingsBase : public UObject
 	GENERATED_BODY()
 
 public:
-	AUDIOEXTENSIONS_API virtual TUniquePtr<ISoundfieldEndpointSettingsProxy> GetProxy() const PURE_VIRTUAL(USoundfieldEndpointSettingsBase::GetProxy, return nullptr;);
+	virtual TUniquePtr<ISoundfieldEndpointSettingsProxy> GetProxy() const PURE_VIRTUAL(USoundfieldEndpointSettingsBase::GetProxy, return nullptr;);
 };
 
 /**
diff --git a/Engine/Source/Runtime/AudioExtensions/Public/ISoundfieldFormat.h b/Engine/Source/Runtime/AudioExtensions/Public/ISoundfieldFormat.h
--- a/Engine/Source/Runtime/AudioExtensions/Public/ISoundfieldFormat.h
+++ b/Engine/Source/Runtime/AudioExtensions/Public/ISoundfieldFormat.h
@@ -125,7 +125,7 @@ class USoundfieldEncodingSettingsBase : public UObject
 	GENERATED_BODY()
 
 public:
-	AUDIOEXTENSIONS_API virtual TUniquePtr<ISoundfieldEncodingSettingsProxy> GetProxy() const PURE_VIRTUAL(USoundfieldEncodingSettingsBase::GetProxy, return nullptr;);
+	virtual TUniquePtr<ISoundfieldEncodingSettingsProxy> GetProxy() const PURE_VIRTUAL(USoundfieldEncodingSettingsBase::GetProxy, return nullptr;);
 };
 
 struct FAudioPluginInitializationParams;
@@ -353,7 +353,7 @@ class USoundfieldEffectSettingsBase : public UObject
 	GENERATED_BODY()
 
 protected:
-	AUDIOEXTENSIONS_API virtual TUniquePtr<ISoundfieldEffectSettingsProxy> GetNewProxy() const PURE_VIRTUAL(USoundfieldEffectSettingsBase::GetProxy, return nullptr;);
+	virtual TUniquePtr<ISoundfieldEffectSettingsProxy> GetNewProxy() const PURE_VIRTUAL(USoundfieldEffectSettingsBase::GetProxy, return nullptr;);
 
 private:
 	// This is called by any engine system that is explicitly marked as a friend.
@@ -396,22 +396,22 @@ protected:
 	/*
 	 * Get the implementation of USoundfieldProcessorSettingsBase that is used for this processor's settings. Will always be called on the CDO.
 	 */
-	AUDIOEXTENSIONS_API virtual bool SupportsFormat(const FName& InFormat) const PURE_VIRTUAL(USoundfieldEncodingSettingsBase::SupportsFormat, return false;);
+	virtual bool SupportsFormat(const FName& InFormat) const PURE_VIRTUAL(USoundfieldEncodingSettingsBase::SupportsFormat, return false;);
 
 	/*
 	 * Get the implementation of USoundfieldProcessorSettingsBase that is used for this processor's settings. Will always be called on the CDO.
 	 */
-	AUDIOEXTENSIONS_API virtual const UClass* GetSettingsClass() const PURE_VIRTUAL(USoundfieldEncodingSettingsBase::GetSettingsClass, return nullptr;);
+	virtual const UClass* GetSettingsClass() const PURE_VIRTUAL(USoundfieldEncodingSettingsBase::GetSettingsClass, return nullptr;);
 	
 	/**
 	 * return the default processor settings we should use when none is provided. Will always be called on the CDO.
 	 */
-	AUDIOEXTENSIONS_API virtual const USoundfieldEffectSettingsBase* GetDefaultSettings() const PURE_VIRTUAL(USoundfieldEncodingSettingsBase::GetDefaultSettings, return nullptr;);
+	virtual const USoundfieldEffectSettingsBase* GetDefaultSettings() const PURE_VIRTUAL(USoundfieldEncodingSettingsBase::GetDefaultSettings, return nullptr;);
 
 	/**
 	 * Spawn a new instance of this processor.
 	 */
-	AUDIOEXTENSIONS_API virtual TUniquePtr<ISoundfieldEffectInstance> GetNewProcessor(const ISoundfieldEncodingSettingsProxy& EncodingSettings) const PURE_VIRTUAL(USoundfieldEncodingSettingsBase::GetProxy, return nullptr;);
+	virtual TUniquePtr<ISoundfieldEffectInstance> GetNewProcessor(const ISoundfieldEncodingSettingsProxy& EncodingSettings) const PURE_VIRTUAL(USoundfieldEncodingSettingsBase::GetProxy, return nullptr;);
 
 private:
 
diff --git a/Engine/Source/Runtime/AudioLink/AudioLinkCore/Public/AudioLinkSettingsAbstract.h b/Engine/Source/Runtime/AudioLink/AudioLinkCore/Public/AudioLinkSettingsAbstract.h
--- a/Engine/Source/Runtime/AudioLink/AudioLinkCore/Public/AudioLinkSettingsAbstract.h
+++ b/Engine/Source/Runtime/AudioLink/AudioLinkCore/Public/AudioLinkSettingsAbstract.h
@@ -35,7 +35,7 @@ class UAudioLinkSettingsAbstract : public UObject
 public:
 	using FSharedSettingsProxyPtr = TSharedPtr<IAudioLinkSettingsProxy, ESPMode::ThreadSafe>;
 
-	AUDIOLINKCORE_API virtual FName GetFactoryName() const PURE_VIRTUAL(UAudioLinkSettingsAbstract::GetFactoryName, return {};);
+	virtual FName GetFactoryName() const PURE_VIRTUAL(UAudioLinkSettingsAbstract::GetFactoryName, return {};);
 	
 	virtual const FSharedSettingsProxyPtr& GetProxy() const
 	{
@@ -60,7 +60,7 @@ protected:
 	}
 #endif //WITH_EDITOR
 
-	AUDIOLINKCORE_API virtual FSharedSettingsProxyPtr MakeProxy() const PURE_VIRTUAL(UAudioLinkSettingsAbstract::MakeProxy, return {};);
+	virtual FSharedSettingsProxyPtr MakeProxy() const PURE_VIRTUAL(UAudioLinkSettingsAbstract::MakeProxy, return {};);
 	mutable FSharedSettingsProxyPtr ProxyInstance;
 };
 
diff --git a/Engine/Source/Runtime/ClothingSystemRuntimeInterface/Public/ClothConfigBase.h b/Engine/Source/Runtime/ClothingSystemRuntimeInterface/Public/ClothConfigBase.h
--- a/Engine/Source/Runtime/ClothingSystemRuntimeInterface/Public/ClothConfigBase.h
+++ b/Engine/Source/Runtime/ClothingSystemRuntimeInterface/Public/ClothConfigBase.h
@@ -23,27 +23,27 @@ public:
 	CLOTHINGSYSTEMRUNTIMEINTERFACE_API virtual ~UClothConfigBase();
 
 	/** Return wherether to pre-compute self collision data. */
-	CLOTHINGSYSTEMRUNTIMEINTERFACE_API virtual bool NeedsSelfCollisionData() const
+	virtual bool NeedsSelfCollisionData() const
 	PURE_VIRTUAL(UClothConfigBase::NeedsSelfCollisionData, return false;);
 
 	/** Return wherether to pre-compute inverse masses. */
-	CLOTHINGSYSTEMRUNTIMEINTERFACE_API virtual bool NeedsInverseMasses() const
+	virtual bool NeedsInverseMasses() const
 	PURE_VIRTUAL(UClothConfigBase::NeedsInverseMasses, return false;);
 
 	/** Return wherether to pre-compute the influences. */
-	CLOTHINGSYSTEMRUNTIMEINTERFACE_API virtual bool NeedsNumInfluences() const
+	virtual bool NeedsNumInfluences() const
 	PURE_VIRTUAL(UClothConfigBase::NeedsNumInfluences, return false;);
 
 	/** Return wherether to pre-compute the long range attachment tethers. */
-	CLOTHINGSYSTEMRUNTIMEINTERFACE_API virtual bool NeedsTethers() const
+	virtual bool NeedsTethers() const
 	PURE_VIRTUAL(UClothConfigBase::NeedsTethers, return false;);
 
 	/** Return the self collision radius to precomute self collision data. */
-	CLOTHINGSYSTEMRUNTIMEINTERFACE_API virtual float GetSelfCollisionRadius() const
+	virtual float GetSelfCollisionRadius() const
 	PURE_VIRTUAL(UClothConfigBase::GetSelfCollisionRadius, return 0.f;);
 
 	/** Return whether tethers need to be calculated using geodesic distances instead of eclidean. */
-	CLOTHINGSYSTEMRUNTIMEINTERFACE_API virtual bool TethersUseGeodesicDistance() const
+	virtual bool TethersUseGeodesicDistance() const
 	PURE_VIRTUAL(UClothConfigBase::TethersUseGeodesicDistance, return false;);
 };
 
diff --git a/Engine/Source/Runtime/ClothingSystemRuntimeInterface/Public/ClothingAssetBase.h b/Engine/Source/Runtime/ClothingSystemRuntimeInterface/Public/ClothingAssetBase.h
--- a/Engine/Source/Runtime/ClothingSystemRuntimeInterface/Public/ClothingAssetBase.h
+++ b/Engine/Source/Runtime/ClothingSystemRuntimeInterface/Public/ClothingAssetBase.h
@@ -26,14 +26,14 @@ public:
 	* @param InSubmeshIdx Submesh in this asset to replace section with
 	* @param InAssetLodIndex Internal clothing LOD to use
 	*/
-	CLOTHINGSYSTEMRUNTIMEINTERFACE_API virtual bool BindToSkeletalMesh(USkeletalMesh* InSkelMesh, const int32 InMeshLodIndex, const int32 InSectionIndex, const int32 InAssetLodIndex)
+	virtual bool BindToSkeletalMesh(USkeletalMesh* InSkelMesh, const int32 InMeshLodIndex, const int32 InSectionIndex, const int32 InAssetLodIndex)
 	PURE_VIRTUAL(UClothingAssetBase::BindToSkeletalMesh, return false;);
 
 	/**
 	* Unbinds this clothing asset from the provided skeletal mesh, will remove all LODs
 	* @param InSkelMesh skeletal mesh to remove from
 	*/
-	CLOTHINGSYSTEMRUNTIMEINTERFACE_API virtual void UnbindFromSkeletalMesh(USkeletalMesh* InSkelMesh)
+	virtual void UnbindFromSkeletalMesh(USkeletalMesh* InSkelMesh)
 	PURE_VIRTUAL(UClothingAssetBase::UnbindFromSkeletalMesh, );
 
 	/**
@@ -41,18 +41,18 @@ public:
 	* @param InSkelMesh skeletal mesh to remove from
 	* @param InMeshLodIndex Mesh LOD to remove this asset from (could still be bound to other LODs)
 	*/
-	CLOTHINGSYSTEMRUNTIMEINTERFACE_API virtual void UnbindFromSkeletalMesh(USkeletalMesh* InSkelMesh, const int32 InMeshLodIndex)
+	virtual void UnbindFromSkeletalMesh(USkeletalMesh* InSkelMesh, const int32 InMeshLodIndex)
 	PURE_VIRTUAL(UClothingAssetBase::UnbindFromSkeletalMesh,);
 
 	/**
 	 * Update all extra LOD deformer mappings.
 	 * This should be called whenever the raytracing LOD bias is changed.
 	 */
-	CLOTHINGSYSTEMRUNTIMEINTERFACE_API virtual void UpdateAllLODBiasMappings(USkeletalMesh* SkeletalMesh)
+	virtual void UpdateAllLODBiasMappings(USkeletalMesh* SkeletalMesh)
 	PURE_VIRTUAL(UClothingAssetBase::UpdateAllLODBiasMappings,);
 
 	/** Add a new LOD class instance. */
-	CLOTHINGSYSTEMRUNTIMEINTERFACE_API virtual int32 AddNewLod()
+	virtual int32 AddNewLod()
 	PURE_VIRTUAL(UClothingAssetBase::AddNewLod(), return INDEX_NONE;);
 
 	/**
@@ -60,7 +60,7 @@ public:
 	*	When we transition between LODs we skin the incoming mesh to the outgoing mesh
 	*	in exactly the same way the render mesh is skinned to create a smooth swap
 	*/
-	CLOTHINGSYSTEMRUNTIMEINTERFACE_API virtual void BuildLodTransitionData()
+	virtual void BuildLodTransitionData()
 	PURE_VIRTUAL(UClothingAssetBase::BuildLodTransitionData(), );
 #endif
 
@@ -69,7 +69,7 @@ public:
 	 * Called on the clothing asset when the base data (physical mesh, config etc.)
 	 * has changed, so any intermediate generated data can be regenerated.
 	 */
-	CLOTHINGSYSTEMRUNTIMEINTERFACE_API virtual void InvalidateAllCachedData()
+	virtual void InvalidateAllCachedData()
 	PURE_VIRTUAL(UClothingAssetBase::InvalidateAllCachedData(),);
 #endif // WITH_EDITORONLY_DATA
 
@@ -79,19 +79,19 @@ public:
 	 * data.
 	 * @param InSkelMesh - The mesh to use to remap the bones
 	 */
-	CLOTHINGSYSTEMRUNTIMEINTERFACE_API virtual void RefreshBoneMapping(USkeletalMesh* InSkelMesh)
+	virtual void RefreshBoneMapping(USkeletalMesh* InSkelMesh)
 	PURE_VIRTUAL(UClothingAssetBase::RefreshBoneMapping, );
 
 	/** Check the validity of a LOD index */
-	CLOTHINGSYSTEMRUNTIMEINTERFACE_API virtual bool IsValidLod(int32 InLodIndex) const
+	virtual bool IsValidLod(int32 InLodIndex) const
 	PURE_VIRTUAL(UClothingAssetBase::IsValidLod(), return false;);
 
 	/** Get the number of LODs defined in the clothing asset */
-	CLOTHINGSYSTEMRUNTIMEINTERFACE_API virtual int32 GetNumLods() const
+	virtual int32 GetNumLods() const
 	PURE_VIRTUAL(UClothingAssetBase::GetNumLods(), return 0;);
 
 	/** Called after all cloth assets sharing the same simulation are added or loaded */
-	CLOTHINGSYSTEMRUNTIMEINTERFACE_API virtual void PostUpdateAllAssets()
+	virtual void PostUpdateAllAssets()
 	PURE_VIRTUAL(UClothingAssetBase::PostUpdateAllAssets(), );
 
 	/** Get the guid identifying this asset */
diff --git a/Engine/Source/Runtime/ClothingSystemRuntimeInterface/Public/ClothingSimulationFactory.h b/Engine/Source/Runtime/ClothingSystemRuntimeInterface/Public/ClothingSimulationFactory.h
--- a/Engine/Source/Runtime/ClothingSystemRuntimeInterface/Public/ClothingSimulationFactory.h
+++ b/Engine/Source/Runtime/ClothingSystemRuntimeInterface/Public/ClothingSimulationFactory.h
@@ -50,33 +50,33 @@ public:
 	static CLOTHINGSYSTEMRUNTIMEINTERFACE_API TSubclassOf<class UClothingSimulationFactory> GetDefaultClothingSimulationFactoryClass();
 
 	// Create a simulation object for a skeletal mesh to use (see IClothingSimulation)
-	CLOTHINGSYSTEMRUNTIMEINTERFACE_API virtual IClothingSimulation* CreateSimulation()
+	virtual IClothingSimulation* CreateSimulation()
 	PURE_VIRTUAL(UClothingSimulationFactory::CreateSimulation, return nullptr;);
 
 	// Destroy a simulation object, guaranteed to be a pointer returned from CreateSimulation for this factory
-	CLOTHINGSYSTEMRUNTIMEINTERFACE_API virtual void DestroySimulation(IClothingSimulation* InSimulation)
+	virtual void DestroySimulation(IClothingSimulation* InSimulation)
 	PURE_VIRTUAL(UClothingSimulationFactory::DestroySimulation, );
 
 	// Given an asset, decide whether this factory can create a simulation to use the data inside
 	// (return false if data is invalid or missing in the case of custom data)
-	CLOTHINGSYSTEMRUNTIMEINTERFACE_API virtual bool SupportsAsset(UClothingAssetBase* InAsset)
+	virtual bool SupportsAsset(UClothingAssetBase* InAsset)
 	PURE_VIRTUAL(UClothingSimulationFactory::SupportsAsset, return false;);
 
 	// Whether or not we provide an interactor object to manipulate the simulation at runtime.
 	// If true is returned then CreateInteractor *must* create a valid object to handle this
-	CLOTHINGSYSTEMRUNTIMEINTERFACE_API virtual bool SupportsRuntimeInteraction()
+	virtual bool SupportsRuntimeInteraction()
 	PURE_VIRTUAL(UClothingSimulationFactory::SupportsRuntimeInteraction, return false;);
 
 	// Creates the runtime interactor object for a clothing simulation. This object will
 	// receive events allowing it to write data to the simulation context in a safe manner
-	CLOTHINGSYSTEMRUNTIMEINTERFACE_API virtual UClothingSimulationInteractor* CreateInteractor()
+	virtual UClothingSimulationInteractor* CreateInteractor()
 	PURE_VIRTUAL(UClothingSimulationFactory::CreateInteractor, return nullptr;);
 
 	// Return the cloth config type for this cloth factory
-	CLOTHINGSYSTEMRUNTIMEINTERFACE_API virtual TArrayView<const TSubclassOf<UClothConfigBase>> GetClothConfigClasses() const
+	virtual TArrayView<const TSubclassOf<UClothConfigBase>> GetClothConfigClasses() const
 	PURE_VIRTUAL(UClothingSimulationFactory::GetClothConfigClasses, return TArrayView<const TSubclassOf<UClothConfigBase>>(););
 
 	// Return an enum of the weight map targets that can be used with this simulation.
-	CLOTHINGSYSTEMRUNTIMEINTERFACE_API virtual const UEnum* GetWeightMapTargetEnum() const
+	virtual const UEnum* GetWeightMapTargetEnum() const
 	PURE_VIRTUAL(UClothingSimulationFactory::GetWeightMapTargetEnum, return nullptr;);
 };
diff --git a/Engine/Source/Runtime/ClothingSystemRuntimeInterface/Public/ClothingSimulationInteractor.h b/Engine/Source/Runtime/ClothingSystemRuntimeInterface/Public/ClothingSimulationInteractor.h
--- a/Engine/Source/Runtime/ClothingSystemRuntimeInterface/Public/ClothingSimulationInteractor.h
+++ b/Engine/Source/Runtime/ClothingSystemRuntimeInterface/Public/ClothingSimulationInteractor.h
@@ -70,42 +70,42 @@ public:
 
 	/** Called to update collision status without restarting the simulation. */
 	UFUNCTION(BlueprintCallable, Category=ClothingSimulation)
-	CLOTHINGSYSTEMRUNTIMEINTERFACE_API virtual void PhysicsAssetUpdated()
+	virtual void PhysicsAssetUpdated()
 	PURE_VIRTUAL(UClothingSimulationInteractor::PhysicsAssetUpdated, );
 
 	/** Called to update the cloth config without restarting the simulation. */
 	UFUNCTION(BlueprintCallable, Category=ClothingSimulation)
-	CLOTHINGSYSTEMRUNTIMEINTERFACE_API virtual void ClothConfigUpdated()
+	virtual void ClothConfigUpdated()
 	PURE_VIRTUAL(UClothingSimulationInteractor::ClothConfigUpdated, );
 
 	/** Set the stiffness of the spring force for the animation drive. */
 	UFUNCTION(BlueprintCallable, Category = ClothingSimulation)
-	CLOTHINGSYSTEMRUNTIMEINTERFACE_API virtual void SetAnimDriveSpringStiffness(float InStiffness)
+	virtual void SetAnimDriveSpringStiffness(float InStiffness)
 	PURE_VIRTUAL(UClothingSimulationInteractor::SetAnimDriveSpringStiffness, );
 
 	/** Set a new gravity override and enable the override. */
 	UFUNCTION(BlueprintCallable, Category = ClothingSimulation)
-	CLOTHINGSYSTEMRUNTIMEINTERFACE_API virtual void EnableGravityOverride(const FVector& InVector)
+	virtual void EnableGravityOverride(const FVector& InVector)
 	PURE_VIRTUAL(UClothingSimulationInteractor::EnableGravityOverride, );
 
 	/** Disable any currently set gravity override. */
 	UFUNCTION(BlueprintCallable, Category = ClothingSimulation)
-	CLOTHINGSYSTEMRUNTIMEINTERFACE_API virtual void DisableGravityOverride()
+	virtual void DisableGravityOverride()
 	PURE_VIRTUAL(UClothingSimulationInteractor::DisableGravityOverride, );
 
 	/** Set the number of time dependent solver iterations. */
 	UFUNCTION(BlueprintCallable, Category = ClothingSimulation)
-	CLOTHINGSYSTEMRUNTIMEINTERFACE_API virtual void SetNumIterations(int32 NumIterations = 1)
+	virtual void SetNumIterations(int32 NumIterations = 1)
 	PURE_VIRTUAL(UClothingSimulationInteractor::SetNumIterations, );
 
 	/** Set the maximum number of solver iterations. */
 	UFUNCTION(BlueprintCallable, Category = ClothingSimulation)
-	CLOTHINGSYSTEMRUNTIMEINTERFACE_API virtual void SetMaxNumIterations(int32 MaxNumIterations = 10)
+	virtual void SetMaxNumIterations(int32 MaxNumIterations = 10)
 	PURE_VIRTUAL(UClothingSimulationInteractor::SetMaxNumIterations, );
 
 	/** Set the number of substeps or subdivisions. */
 	UFUNCTION(BlueprintCallable, Category = ClothingSimulation, Meta = (Keywords = "Subdivisions"))
-	CLOTHINGSYSTEMRUNTIMEINTERFACE_API virtual void SetNumSubsteps(int32 NumSubsteps = 1)
+	virtual void SetNumSubsteps(int32 NumSubsteps = 1)
 	PURE_VIRTUAL(UClothingSimulationInteractor::SetNumSubsteps, );
 
 	// Base clothing simulations interface
@@ -149,7 +149,7 @@ public:
 	TMap<FName, TObjectPtr<UClothingInteractor>> ClothingInteractors;
 
 protected:
-	CLOTHINGSYSTEMRUNTIMEINTERFACE_API virtual UClothingInteractor* CreateClothingInteractor()
+	virtual UClothingInteractor* CreateClothingInteractor()
 	PURE_VIRTUAL(UClothingSimulationInteractor::CreateClothingInteractor, return nullptr;);
 
 private:
diff --git a/Engine/Source/Runtime/Core/Public/Serialization/BitArchive.h b/Engine/Source/Runtime/Core/Public/Serialization/BitArchive.h
--- a/Engine/Source/Runtime/Core/Public/Serialization/BitArchive.h
+++ b/Engine/Source/Runtime/Core/Public/Serialization/BitArchive.h
@@ -17,5 +17,5 @@ public:
 	 */
 	CORE_API FBitArchive();
 
-	CORE_API virtual void SerializeBitsWithOffset( void* Src, int32 SourceBit, int64 LengthBits ) PURE_VIRTUAL(FBitArchive::SerializeBitsWithOffset,);
+	virtual void SerializeBitsWithOffset( void* Src, int32 SourceBit, int64 LengthBits ) PURE_VIRTUAL(FBitArchive::SerializeBitsWithOffset,);
 };
diff --git a/Engine/Source/Runtime/CoreUObject/Public/UObject/UnrealType.h b/Engine/Source/Runtime/CoreUObject/Public/UObject/UnrealType.h
--- a/Engine/Source/Runtime/CoreUObject/Public/UObject/UnrealType.h
+++ b/Engine/Source/Runtime/CoreUObject/Public/UObject/UnrealType.h
@@ -284,7 +284,7 @@ PRAGMA_ENABLE_DEPRECATION_WARNINGS
 	 * @param	ExtendedTypeText	for property types which use templates, will be filled in with the type
 	 * @param	CPPExportFlags		flags for modifying the behavior of the export
 	 */
-	COREUOBJECT_API virtual FString GetCPPType( FString* ExtendedTypeText=NULL, uint32 CPPExportFlags=0 ) const PURE_VIRTUAL(FProperty::GetCPPType,return TEXT(""););
+	virtual FString GetCPPType( FString* ExtendedTypeText=NULL, uint32 CPPExportFlags=0 ) const PURE_VIRTUAL(FProperty::GetCPPType,return TEXT(""););
 	// End of UHT interface
 
 #if WITH_EDITORONLY_DATA
@@ -455,7 +455,7 @@ public:
 	 *
 	 * @return	true if the property values are identical
 	 */
-	COREUOBJECT_API virtual bool Identical( const void* A, const void* B, uint32 PortFlags=0 ) const PURE_VIRTUAL(FProperty::Identical,return false;);
+	virtual bool Identical( const void* A, const void* B, uint32 PortFlags=0 ) const PURE_VIRTUAL(FProperty::Identical,return false;);
 
 	/**
 	 * Determines whether the property values are identical.
@@ -521,7 +521,7 @@ public:
 		}
 	}
 
-	COREUOBJECT_API virtual void SerializeItem(FStructuredArchive::FSlot Slot, void* Value, void const* Defaults = NULL) const PURE_VIRTUAL(FProperty::SerializeItem, );
+	virtual void SerializeItem(FStructuredArchive::FSlot Slot, void* Value, void const* Defaults = NULL) const PURE_VIRTUAL(FProperty::SerializeItem, );
 	COREUOBJECT_API virtual bool NetSerializeItem( FArchive& Ar, UPackageMap* Map, void* Data, TArray<uint8> * MetaData = NULL ) const;
 	COREUOBJECT_API virtual bool SupportsNetSharedSerialization() const;
 
@@ -658,8 +658,8 @@ public:
 #endif
 protected:
 
-	COREUOBJECT_API virtual void ExportText_Internal(FString& ValueStr, const void* PropertyValueOrContainer, EPropertyPointerType PointerType, const void* DefaultValue, UObject* Parent, int32 PortFlags, UObject* ExportRootScope = nullptr) const PURE_VIRTUAL(FProperty::ExportText, );
-	COREUOBJECT_API virtual const TCHAR* ImportText_Internal(const TCHAR* Buffer, void* ContainerOrPropertyPtr, EPropertyPointerType PointerType, UObject* OwnerObject, int32 PortFlags, FOutputDevice* ErrorText) const PURE_VIRTUAL(FProperty::ImportText, return nullptr;);
+	virtual void ExportText_Internal(FString& ValueStr, const void* PropertyValueOrContainer, EPropertyPointerType PointerType, const void* DefaultValue, UObject* Parent, int32 PortFlags, UObject* ExportRootScope = nullptr) const PURE_VIRTUAL(FProperty::ExportText, );
+	virtual const TCHAR* ImportText_Internal(const TCHAR* Buffer, void* ContainerOrPropertyPtr, EPropertyPointerType PointerType, UObject* OwnerObject, int32 PortFlags, FOutputDevice* ErrorText) const PURE_VIRTUAL(FProperty::ImportText, return nullptr;);
 
 public:
 	
@@ -1275,7 +1275,7 @@ public:
 	 * @see Optional.h - HasIntrusiveUnsetOptionalState
 	 * @see FOptionalProperty
 	 */
-	COREUOBJECT_API virtual bool HasIntrusiveUnsetOptionalState() const PURE_VIRTUAL(FProperty::HasIntrusiveUnsetOptionalState, return false;)
+	virtual bool HasIntrusiveUnsetOptionalState() const PURE_VIRTUAL(FProperty::HasIntrusiveUnsetOptionalState, return false;)
 
 	/**
 	 * Initialize the value at the given address to an unset TOptional using an intrusive state rather than a trailing boolean.
@@ -1793,9 +1793,9 @@ class FNumericProperty : public FProperty
 	UE_API static int64 ReadEnumAsInt64(FStructuredArchive::FSlot Slot, UStruct* DefaultsStruct, const FPropertyTag& Tag);
 
 private:
-	UE_API virtual bool CanHoldDoubleValueInternal  (double Value) const PURE_VIRTUAL(FNumericProperty::CanHoldDoubleValueInternal,   return false;);
-	UE_API virtual bool CanHoldSignedValueInternal  (int64  Value) const PURE_VIRTUAL(FNumericProperty::CanHoldSignedValueInternal,   return false;);
-	UE_API virtual bool CanHoldUnsignedValueInternal(uint64 Value) const PURE_VIRTUAL(FNumericProperty::CanHoldUnsignedValueInternal, return false;);
+	virtual bool CanHoldDoubleValueInternal  (double Value) const PURE_VIRTUAL(FNumericProperty::CanHoldDoubleValueInternal,   return false;);
+	virtual bool CanHoldSignedValueInternal  (int64  Value) const PURE_VIRTUAL(FNumericProperty::CanHoldSignedValueInternal,   return false;);
+	virtual bool CanHoldUnsignedValueInternal(uint64 Value) const PURE_VIRTUAL(FNumericProperty::CanHoldUnsignedValueInternal, return false;);
 };
 
 template<typename InTCppType>
@@ -2565,7 +2565,7 @@ public:
 public:
 
 	UE_DEPRECATED(5.7, "GetCPPTypeCustom is deprecated, and object properties now implement GetCPPType directly.")
-	UE_API FString GetCPPTypeCustom(FString* ExtendedTypeText, uint32 CPPExportFlags, const FString& InnerNativeTypeName) const { return {}; }
+	FString GetCPPTypeCustom(FString* ExtendedTypeText, uint32 CPPExportFlags, const FString& InnerNativeTypeName) const { return {}; }
 
 	/**
 	 * Parses a text buffer into an object reference.
@@ -6098,18 +6098,18 @@ public:
 	UE_API virtual EConvertFromTypeResult ConvertFromType(const FPropertyTag& Tag, FStructuredArchive::FSlot Slot, uint8* Data, UStruct* DefaultsStruct, const uint8* Defaults) override;
 	// End of FProperty interface
 
-	UE_API virtual const FMulticastScriptDelegate* GetMulticastDelegate(const void* PropertyValue) const PURE_VIRTUAL(FMulticastDelegateProperty::GetMulticastDelegate, return nullptr;);
-	UE_API virtual void SetMulticastDelegate(void* PropertyValue, FMulticastScriptDelegate ScriptDelegate) const PURE_VIRTUAL(FMulticastDelegateProperty::SetMulticastDelegate, );
+	virtual const FMulticastScriptDelegate* GetMulticastDelegate(const void* PropertyValue) const PURE_VIRTUAL(FMulticastDelegateProperty::GetMulticastDelegate, return nullptr;);
+	virtual void SetMulticastDelegate(void* PropertyValue, FMulticastScriptDelegate ScriptDelegate) const PURE_VIRTUAL(FMulticastDelegateProperty::SetMulticastDelegate, );
 
-	UE_API virtual void AddDelegate(FScriptDelegate ScriptDelegate, UObject* Parent = nullptr, void* PropertyValue = nullptr) const PURE_VIRTUAL(FMulticastDelegateProperty::AddDelegate, );
-	UE_API virtual void RemoveDelegate(const FScriptDelegate& ScriptDelegate, UObject* Parent = nullptr, void* PropertyValue = nullptr) const PURE_VIRTUAL(FMulticastDelegateProperty::RemoveDelegate, );
-	UE_API virtual void ClearDelegate(UObject* Parent = nullptr, void* PropertyValue = nullptr)  const PURE_VIRTUAL(FMulticastDelegateProperty::ClearDelegate, );
+	virtual void AddDelegate(FScriptDelegate ScriptDelegate, UObject* Parent = nullptr, void* PropertyValue = nullptr) const PURE_VIRTUAL(FMulticastDelegateProperty::AddDelegate, );
+	virtual void RemoveDelegate(const FScriptDelegate& ScriptDelegate, UObject* Parent = nullptr, void* PropertyValue = nullptr) const PURE_VIRTUAL(FMulticastDelegateProperty::RemoveDelegate, );
+	virtual void ClearDelegate(UObject* Parent = nullptr, void* PropertyValue = nullptr)  const PURE_VIRTUAL(FMulticastDelegateProperty::ClearDelegate, );
 
 protected:
 	friend class FProperty;
 
 	UE_API static FMulticastScriptDelegate EmptyDelegate;
-	UE_API virtual FMulticastScriptDelegate& GetMulticastScriptDelegate(const void* PropertyValue, int32 Index) const PURE_VIRTUAL(FMulticastDelegateProperty::GetMulticastScriptDelegate, return EmptyDelegate;);
+	virtual FMulticastScriptDelegate& GetMulticastScriptDelegate(const void* PropertyValue, int32 Index) const PURE_VIRTUAL(FMulticastDelegateProperty::GetMulticastScriptDelegate, return EmptyDelegate;);
 
 
 	UE_API const TCHAR* ImportText_Add( const TCHAR* Buffer, void* PropertyValue, int32 PortFlags, UObject* Parent, FOutputDevice* ErrorText ) const;
diff --git a/Engine/Source/Runtime/Datasmith/DatasmithCore/Public/DatasmithCloth.h b/Engine/Source/Runtime/Datasmith/DatasmithCore/Public/DatasmithCloth.h
--- a/Engine/Source/Runtime/Datasmith/DatasmithCore/Public/DatasmithCloth.h
+++ b/Engine/Source/Runtime/Datasmith/DatasmithCore/Public/DatasmithCloth.h
@@ -110,16 +110,16 @@ public:
 	virtual ~UDatasmithClothAssetFactory() = default;
 
 	UE_DEPRECATED(5.5, "The experimental Cloth importer is no longer supported.")
-	UE_API virtual UObject* CreateClothAsset(UObject* Outer, const FName& Name, EObjectFlags Flags) const
+	virtual UObject* CreateClothAsset(UObject* Outer, const FName& Name, EObjectFlags Flags) const
 	PURE_VIRTUAL(UDatasmithClothAssetFactory::CreateClothAsset, return nullptr;);
 
 	UE_DEPRECATED(5.5, "The experimental Cloth importer is no longer supported.")
-	UE_API virtual UObject* DuplicateClothAsset(UObject* ClothAsset, UObject* Outer, const FName& Name) const
+	virtual UObject* DuplicateClothAsset(UObject* ClothAsset, UObject* Outer, const FName& Name) const
 	PURE_VIRTUAL(UDatasmithClothAssetFactory::DuplicateClothAsset, return nullptr;);
 
 PRAGMA_DISABLE_DEPRECATION_WARNINGS
 	UE_DEPRECATED(5.5, "The experimental Cloth importer is no longer supported.")
-	UE_API virtual void InitializeClothAsset(UObject* ClothAsset, const FDatasmithCloth& DatasmithCloth) const
+	virtual void InitializeClothAsset(UObject* ClothAsset, const FDatasmithCloth& DatasmithCloth) const
 	PURE_VIRTUAL(UDatasmithClothAssetFactory::InitializeClothAsset, );
 PRAGMA_ENABLE_DEPRECATION_WARNINGS
 };
@@ -138,11 +138,11 @@ public:
 	virtual ~UDatasmithClothComponentFactory() = default;
 
 	UE_DEPRECATED(5.5, "The experimental Cloth importer is no longer supported.")
-	UE_API virtual USceneComponent* CreateClothComponent(UObject* Outer) const
+	virtual USceneComponent* CreateClothComponent(UObject* Outer) const
 	PURE_VIRTUAL(UDatasmithClothComponentFactory::CreateClothComponent, return nullptr;);
 
 	UE_DEPRECATED(5.5, "The experimental Cloth importer is no longer supported.")
-	UE_API virtual void InitializeClothComponent(USceneComponent* ClothComponent, UObject* ClothAsset, USceneComponent* RootComponent) const
+	virtual void InitializeClothComponent(USceneComponent* ClothComponent, UObject* ClothAsset, USceneComponent* RootComponent) const
 	PURE_VIRTUAL(UDatasmithClothComponentFactory::InitializeClothComponent, );
 };
 
diff --git a/Engine/Source/Runtime/Engine/Classes/AI/AISystemBase.h b/Engine/Source/Runtime/Engine/Classes/AI/AISystemBase.h
--- a/Engine/Source/Runtime/Engine/Classes/AI/AISystemBase.h
+++ b/Engine/Source/Runtime/Engine/Classes/AI/AISystemBase.h
@@ -21,7 +21,7 @@ class UAISystemBase : public UObject
 	 * Called when world initializes all actors and prepares them to start gameplay
 	 * @param bTimeGotReset whether the WorldSettings's TimeSeconds are reset to zero
 	 */
-	ENGINE_API virtual void InitializeActorsForPlay(bool bTimeGotReset) PURE_VIRTUAL(UAISystemBase::InitializeActorsForPlay, );
+	virtual void InitializeActorsForPlay(bool bTimeGotReset) PURE_VIRTUAL(UAISystemBase::InitializeActorsForPlay, );
 
 	/**
 	 * Event called on world origin location changes
@@ -29,7 +29,7 @@ class UAISystemBase : public UObject
 	 * @param	OldOriginLocation			Previous world origin location
 	 * @param	NewOriginLocation			New world origin location
 	 */
-	ENGINE_API virtual void WorldOriginLocationChanged(FIntVector OldOriginLocation, FIntVector NewOriginLocation) PURE_VIRTUAL(UAISystemBase::WorldOriginLocationChanged, );
+	virtual void WorldOriginLocationChanged(FIntVector OldOriginLocation, FIntVector NewOriginLocation) PURE_VIRTUAL(UAISystemBase::WorldOriginLocationChanged, );
 
 	/**
 	 * Called by UWorld::CleanupWorld.
diff --git a/Engine/Source/Runtime/Engine/Classes/AI/Navigation/NavCollisionBase.h b/Engine/Source/Runtime/Engine/Classes/AI/Navigation/NavCollisionBase.h
--- a/Engine/Source/Runtime/Engine/Classes/AI/Navigation/NavCollisionBase.h
+++ b/Engine/Source/Runtime/Engine/Classes/AI/Navigation/NavCollisionBase.h
@@ -50,21 +50,21 @@ public:
 
 	/** Tries to read data from DDC, and if that fails gathers navigation
 	*	collision data, stores it and uploads to DDC */
-	ENGINE_API virtual void Setup(class UBodySetup* BodySetup) PURE_VIRTUAL(UNavCollisionBase::Setup, );
+	virtual void Setup(class UBodySetup* BodySetup) PURE_VIRTUAL(UNavCollisionBase::Setup, );
 
-	[[nodiscard]] ENGINE_API virtual FBox GetBounds() const PURE_VIRTUAL(UNavCollisionBase::GetBounds, static FBox InvalidBox; return InvalidBox; );
+	[[nodiscard]] virtual FBox GetBounds() const PURE_VIRTUAL(UNavCollisionBase::GetBounds, static FBox InvalidBox; return InvalidBox; );
 
 	/** Export collision data */
-	ENGINE_API virtual bool ExportGeometry(const FTransform& LocalToWorld, FNavigableGeometryExport& GeoExport) const PURE_VIRTUAL(UNavCollisionBase::ExportGeometry, return false; );
+	virtual bool ExportGeometry(const FTransform& LocalToWorld, FNavigableGeometryExport& GeoExport) const PURE_VIRTUAL(UNavCollisionBase::ExportGeometry, return false; );
 
 	/** Get data for dynamic obstacle */
-	ENGINE_API virtual void GetNavigationModifier(FCompositeNavModifier& Modifier, const FTransform& LocalToWorld) PURE_VIRTUAL(UNavCollisionBase::GetNavigationModifier, );
+	virtual void GetNavigationModifier(FCompositeNavModifier& Modifier, const FTransform& LocalToWorld) PURE_VIRTUAL(UNavCollisionBase::GetNavigationModifier, );
 
 	/** draw cylinder and box collisions */
 	virtual void DrawSimpleGeom(FPrimitiveDrawInterface* PDI, const FTransform& Transform, const FColor Color) {}
 
 #if WITH_EDITOR
-	ENGINE_API virtual void InvalidateCollision() PURE_VIRTUAL(UNavCollisionBase::InvalidateCollision, );
+	virtual void InvalidateCollision() PURE_VIRTUAL(UNavCollisionBase::InvalidateCollision, );
 #endif // WITH_EDITOR
 
 	bool IsDynamicObstacle() const { return bIsDynamicObstacle; }
diff --git a/Engine/Source/Runtime/Engine/Classes/AI/NavigationSystemBase.h b/Engine/Source/Runtime/Engine/Classes/AI/NavigationSystemBase.h
--- a/Engine/Source/Runtime/Engine/Classes/AI/NavigationSystemBase.h
+++ b/Engine/Source/Runtime/Engine/Classes/AI/NavigationSystemBase.h
@@ -248,11 +248,11 @@ class UNavigationSystemBase : public UObject
 public:
 	virtual ~UNavigationSystemBase(){}
 
-	ENGINE_API virtual void Tick(float DeltaSeconds) PURE_VIRTUAL(UNavigationSystemBase::Tick, );
-	ENGINE_API virtual void CleanUp(const FNavigationSystem::ECleanupMode Mode) PURE_VIRTUAL(UNavigationSystemBase::CleanUp, );
-	ENGINE_API virtual void Configure(const UNavigationSystemConfig& Config) PURE_VIRTUAL(UNavigationSystemBase::Configure, );
+	virtual void Tick(float DeltaSeconds) PURE_VIRTUAL(UNavigationSystemBase::Tick, );
+	virtual void CleanUp(const FNavigationSystem::ECleanupMode Mode) PURE_VIRTUAL(UNavigationSystemBase::CleanUp, );
+	virtual void Configure(const UNavigationSystemConfig& Config) PURE_VIRTUAL(UNavigationSystemBase::Configure, );
 	/** Called when there's a need to extend current navigation system's config with information in NewConfig */
-	ENGINE_API virtual void AppendConfig(const UNavigationSystemConfig& NewConfig) PURE_VIRTUAL(UNavigationSystemBase::AppendConfig, );
+	virtual void AppendConfig(const UNavigationSystemConfig& NewConfig) PURE_VIRTUAL(UNavigationSystemBase::AppendConfig, );
 
 	/**
 	*	Called when owner-UWorld initializes actors
@@ -261,9 +261,9 @@ public:
 
 	virtual bool IsNavigationBuilt(const AWorldSettings* Settings) const { return false; }
 
-	ENGINE_API virtual void ApplyWorldOffset(const FVector& InOffset, bool bWorldShift) PURE_VIRTUAL(UNavigationSystemBase::ApplyWorldOffset, );
+	virtual void ApplyWorldOffset(const FVector& InOffset, bool bWorldShift) PURE_VIRTUAL(UNavigationSystemBase::ApplyWorldOffset, );
 
-	ENGINE_API virtual void InitializeForWorld(UWorld& World, FNavigationSystemRunMode Mode) PURE_VIRTUAL(UNavigationSystemBase::InitializeForWorld, );
+	virtual void InitializeForWorld(UWorld& World, FNavigationSystemRunMode Mode) PURE_VIRTUAL(UNavigationSystemBase::InitializeForWorld, );
 
 	/** 
 	 *	If you're using NavigationSysstem module consider calling 
@@ -272,25 +272,25 @@ public:
 	 */
 	virtual INavigationDataInterface* GetMainNavData() const { return nullptr; }
 
-	ENGINE_API virtual void SetBuildBounds(const FBox& Bounds) PURE_VIRTUAL(UNavigationSystemBase::SetBuildBounds, );
+	virtual void SetBuildBounds(const FBox& Bounds) PURE_VIRTUAL(UNavigationSystemBase::SetBuildBounds, );
 
-	ENGINE_API virtual FBox GetNavigableWorldBounds() const PURE_VIRTUAL(UNavigationSystemBase::GetNavigableWorldBounds, return FBox(ForceInit););
+	virtual FBox GetNavigableWorldBounds() const PURE_VIRTUAL(UNavigationSystemBase::GetNavigableWorldBounds, return FBox(ForceInit););
 	
-	ENGINE_API virtual bool ContainsNavData(const FBox& Bounds) const PURE_VIRTUAL(UNavigationSystemBase::ContainsNavData, return false;);
-	ENGINE_API virtual FBox ComputeNavDataBounds() const PURE_VIRTUAL(UNavigationSystemBase::GetNavigableWorldBounds, return FBox(ForceInit););
+	virtual bool ContainsNavData(const FBox& Bounds) const PURE_VIRTUAL(UNavigationSystemBase::ContainsNavData, return false;);
+	virtual FBox ComputeNavDataBounds() const PURE_VIRTUAL(UNavigationSystemBase::GetNavigableWorldBounds, return FBox(ForceInit););
 	
 	virtual void AddNavigationDataChunk(class ANavigationDataChunkActor& DataChunkActor) {}
 	virtual void RemoveNavigationDataChunk(class ANavigationDataChunkActor& DataChunkActor) {}
 	virtual void FillNavigationDataChunkActor(const FBox& QueryBounds, class ANavigationDataChunkActor& DataChunkActor, FBox& OutTilesBounds) {}
 
-	ENGINE_API virtual bool IsWorldInitDone() const PURE_VIRTUAL(UNavigationSystemBase::IsWorldInitDone, return false;);
+	virtual bool IsWorldInitDone() const PURE_VIRTUAL(UNavigationSystemBase::IsWorldInitDone, return false;);
 
 	/**
 	 * Indicates whether a navigation system instance should be created for a given world.
 	 * @param World The world in which the navigation would be used
 	 * @return whether the navigation system should be created for the specified world or not
 	 */
-	ENGINE_API virtual bool ShouldCreateNavigationSystemInstance(const UWorld* World) const PURE_VIRTUAL(UNavigationSystemBase::ShouldCreateNavigationSystemInstance, return true;);
+	virtual bool ShouldCreateNavigationSystemInstance(const UWorld* World) const PURE_VIRTUAL(UNavigationSystemBase::ShouldCreateNavigationSystemInstance, return true;);
 
 	static ENGINE_API FNavigationSystem::FOnNavigationInitSignature& OnNavigationInitStartStaticDelegate();
 	static ENGINE_API FNavigationSystem::FOnNavigationInitSignature& OnNavigationInitDoneStaticDelegate();
diff --git a/Engine/Source/Runtime/Engine/Classes/Animation/AnimBoneCompressionCodec.h b/Engine/Source/Runtime/Engine/Classes/Animation/AnimBoneCompressionCodec.h
--- a/Engine/Source/Runtime/Engine/Classes/Animation/AnimBoneCompressionCodec.h
+++ b/Engine/Source/Runtime/Engine/Classes/Animation/AnimBoneCompressionCodec.h
@@ -41,7 +41,7 @@ class UAnimBoneCompressionCodec : public UObject
 	virtual bool IsHighFidelity(const FCompressibleAnimData& CompressibleAnimData) const { return false; }
 
 	/** Compresses the curve data from an animation sequence. */
-	ENGINE_API virtual bool Compress(const FCompressibleAnimData& CompressibleAnimData, FCompressibleAnimDataResult& OutResult) PURE_VIRTUAL(UAnimCurveCompressionCodec::Compress, return false;);
+	virtual bool Compress(const FCompressibleAnimData& CompressibleAnimData, FCompressibleAnimDataResult& OutResult) PURE_VIRTUAL(UAnimCurveCompressionCodec::Compress, return false;);
 
 	/**
 	 * Estimates the peak memory usage in bytes needed to compress the given data with this codec.
@@ -84,7 +84,7 @@ class UAnimBoneCompressionCodec : public UObject
 	ENGINE_API virtual FString GetCodecDDCHandle() const;
 
 	/** Allocates the necessary anim data structure used for decompression. The engine will serialize the data in/out and bind it. */
-	ENGINE_API virtual TUniquePtr<ICompressedAnimData> AllocateAnimData() const PURE_VIRTUAL(UAnimCurveCompressionCodec::AllocateAnimData, return TUniquePtr<ICompressedAnimData>(););
+	virtual TUniquePtr<ICompressedAnimData> AllocateAnimData() const PURE_VIRTUAL(UAnimCurveCompressionCodec::AllocateAnimData, return TUniquePtr<ICompressedAnimData>(););
 
 	/**
 	 * Handles Byte-swapping incoming animation data from a MemoryReader
@@ -93,7 +93,7 @@ class UAnimBoneCompressionCodec : public UObject
 	 * @param	MemoryReader		The MemoryReader object to read from.
 	 * @param	SourceArVersion		The version of the archive that the data is coming from.
 	 */
-	ENGINE_API virtual void ByteSwapIn(ICompressedAnimData& AnimData, TArrayView<uint8> CompressedData, FMemoryReader& MemoryStream) const PURE_VIRTUAL(UAnimCurveCompressionCodec::ByteSwapIn, );
+	virtual void ByteSwapIn(ICompressedAnimData& AnimData, TArrayView<uint8> CompressedData, FMemoryReader& MemoryStream) const PURE_VIRTUAL(UAnimCurveCompressionCodec::ByteSwapIn, );
 
 	/**
 	 * Handles Byte-swapping outgoing animation data to an array of BYTEs
@@ -102,17 +102,17 @@ class UAnimBoneCompressionCodec : public UObject
 	 * @param	SerializedData		The output buffer.
 	 * @param	bForceByteSwapping	true is byte swapping is not optional.
 	 */
-	ENGINE_API virtual void ByteSwapOut(ICompressedAnimData& AnimData, TArrayView<uint8> CompressedData, FMemoryWriter& MemoryStream) const PURE_VIRTUAL(UAnimCurveCompressionCodec::ByteSwapOut, );
+	virtual void ByteSwapOut(ICompressedAnimData& AnimData, TArrayView<uint8> CompressedData, FMemoryWriter& MemoryStream) const PURE_VIRTUAL(UAnimCurveCompressionCodec::ByteSwapOut, );
 
 	/**
 	 * Decompresses all the specified bone tracks.
 	 * The caller is responsible for pre-filling the output pose with sensible values (e.g. reference/bind/additive identity pose) as
 	 * the codec will only decompress and write out tracks that are contained in the compressed data.
 	 */
-	ENGINE_API virtual void DecompressPose(FAnimSequenceDecompressionContext& DecompContext, const BoneTrackArray& RotationPairs, const BoneTrackArray& TranslationPairs, const BoneTrackArray& ScalePairs, TArrayView<FTransform>& OutAtoms) const PURE_VIRTUAL(UAnimCurveCompressionCodec::DecompressPose, );
+	virtual void DecompressPose(FAnimSequenceDecompressionContext& DecompContext, const BoneTrackArray& RotationPairs, const BoneTrackArray& TranslationPairs, const BoneTrackArray& ScalePairs, TArrayView<FTransform>& OutAtoms) const PURE_VIRTUAL(UAnimCurveCompressionCodec::DecompressPose, );
 
 	/** Decompress a single bone. */
-	ENGINE_API virtual void DecompressBone(FAnimSequenceDecompressionContext& DecompContext, int32 TrackIndex, FTransform& OutAtom) const PURE_VIRTUAL(UAnimCurveCompressionCodec::DecompressBone, );
+	virtual void DecompressBone(FAnimSequenceDecompressionContext& DecompContext, int32 TrackIndex, FTransform& OutAtom) const PURE_VIRTUAL(UAnimCurveCompressionCodec::DecompressBone, );
 
 	/**
 	 * Decompresses all the specified bone tracks.
diff --git a/Engine/Source/Runtime/Engine/Classes/Animation/AnimCurveCompressionCodec.h b/Engine/Source/Runtime/Engine/Classes/Animation/AnimCurveCompressionCodec.h
--- a/Engine/Source/Runtime/Engine/Classes/Animation/AnimCurveCompressionCodec.h
+++ b/Engine/Source/Runtime/Engine/Classes/Animation/AnimCurveCompressionCodec.h
@@ -55,7 +55,7 @@ class UAnimCurveCompressionCodec : public UObject
 	virtual bool IsCodecValid() const { return true; }
 
 	/** Compresses the curve data from an animation sequence. */
-	ENGINE_API virtual bool Compress(const FCompressibleAnimData& AnimSeq, FAnimCurveCompressionResult& OutResult) PURE_VIRTUAL(UAnimCurveCompressionCodec::Compress, return false;);
+	virtual bool Compress(const FCompressibleAnimData& AnimSeq, FAnimCurveCompressionResult& OutResult) PURE_VIRTUAL(UAnimCurveCompressionCodec::Compress, return false;);
 
 	/**
 	 * Estimates the peak memory usage in bytes needed to compress the given data with this codec.
@@ -82,14 +82,14 @@ class UAnimCurveCompressionCodec : public UObject
 	 * Note: Codecs should _NOT_ rely on any member properties during decompression. Decompression
 	 * behavior should entirely be driven by code and the compressed data.
 	 */
-	ENGINE_API virtual void DecompressCurves(const FCompressedAnimSequence& AnimSeq, FBlendedCurve& Curves, float CurrentTime) const PURE_VIRTUAL(UAnimCurveCompressionCodec::DecompressCurves, );
+	virtual void DecompressCurves(const FCompressedAnimSequence& AnimSeq, FBlendedCurve& Curves, float CurrentTime) const PURE_VIRTUAL(UAnimCurveCompressionCodec::DecompressCurves, );
 
 	/*
 	 * Decompress a single curve.
 	 * Note: Codecs should _NOT_ rely on any member properties during decompression. Decompression
 	 * behavior should entirely be driven by code and the compressed data.
 	 */
-	ENGINE_API virtual float DecompressCurve(const FCompressedAnimSequence& AnimSeq, FName CurveName, float CurrentTime) const PURE_VIRTUAL(UAnimCurveCompressionCodec::DecompressCurve, return 0.0f;);
+	virtual float DecompressCurve(const FCompressedAnimSequence& AnimSeq, FName CurveName, float CurrentTime) const PURE_VIRTUAL(UAnimCurveCompressionCodec::DecompressCurve, return 0.0f;);
 	
 	UE_DEPRECATED(5.3, "Please use DecompressCurve that takes an FName.")
     virtual float DecompressCurve(const FCompressedAnimSequence& AnimSeq, SmartName::UID_Type CurveUID, float CurrentTime) const { return 0.0f; }
diff --git a/Engine/Source/Runtime/Engine/Classes/Animation/AnimNode_CustomProperty.h b/Engine/Source/Runtime/Engine/Classes/Animation/AnimNode_CustomProperty.h
--- a/Engine/Source/Runtime/Engine/Classes/Animation/AnimNode_CustomProperty.h
+++ b/Engine/Source/Runtime/Engine/Classes/Animation/AnimNode_CustomProperty.h
@@ -86,7 +86,7 @@ protected:
 	ENGINE_API virtual void PropagateInputProperties(const UObject* InSourceInstance);
 
 	/** Get Target Class */
-	ENGINE_API virtual UClass* GetTargetClass() const PURE_VIRTUAL(FAnimNode_CustomProperty::GetTargetClass, return nullptr;);
+	virtual UClass* GetTargetClass() const PURE_VIRTUAL(FAnimNode_CustomProperty::GetTargetClass, return nullptr;);
 
 #if WITH_EDITOR
 	/**
diff --git a/Engine/Source/Runtime/Engine/Classes/Animation/AnimSequenceBase.h b/Engine/Source/Runtime/Engine/Classes/Animation/AnimSequenceBase.h
--- a/Engine/Source/Runtime/Engine/Classes/Animation/AnimSequenceBase.h
+++ b/Engine/Source/Runtime/Engine/Classes/Animation/AnimSequenceBase.h
@@ -182,7 +182,7 @@ public:
 	* @param	OutCurve			Curves to fill
 	* @param	ExtractionContext	Extraction Context (position, looping, root motion, etc.)
 	*/	
-	ENGINE_API virtual void GetAnimationPose(FAnimationPoseData& OutPoseData, const FAnimExtractContext& ExtractionContext) const
+	virtual void GetAnimationPose(FAnimationPoseData& OutPoseData, const FAnimExtractContext& ExtractionContext) const
 		PURE_VIRTUAL(UAnimSequenceBase::GetAnimationPose, );
 	
 	ENGINE_API virtual void HandleAssetPlayerTickedInternal(FAnimAssetTickContext &Context, const float PreviousTime, const float MoveDelta, const FAnimTickRecord &Instance, struct FAnimNotifyQueue& NotifyQueue) const;
diff --git a/Engine/Source/Runtime/Engine/Classes/Animation/MeshDeformer.h b/Engine/Source/Runtime/Engine/Classes/Animation/MeshDeformer.h
--- a/Engine/Source/Runtime/Engine/Classes/Animation/MeshDeformer.h
+++ b/Engine/Source/Runtime/Engine/Classes/Animation/MeshDeformer.h
@@ -21,12 +21,12 @@ class UMeshDeformer : public UObject
 
 public:
 	/** Create persistent settings for the mesh deformer instance */ 
-	ENGINE_API virtual UMeshDeformerInstanceSettings* CreateSettingsInstance(
+	virtual UMeshDeformerInstanceSettings* CreateSettingsInstance(
 		UMeshComponent* InMeshComponent
 		) PURE_VIRTUAL(, return nullptr;);
 	
 	/** Create an instance of the mesh deformer. */
-	ENGINE_API virtual UMeshDeformerInstance* CreateInstance(
+	virtual UMeshDeformerInstance* CreateInstance(
 		UMeshComponent* InMeshComponent,
 		UMeshDeformerInstanceSettings* InSettings
 		) PURE_VIRTUAL(, return nullptr;);
diff --git a/Engine/Source/Runtime/Engine/Classes/Animation/MeshDeformerInstance.h b/Engine/Source/Runtime/Engine/Classes/Animation/MeshDeformerInstance.h
--- a/Engine/Source/Runtime/Engine/Classes/Animation/MeshDeformerInstance.h
+++ b/Engine/Source/Runtime/Engine/Classes/Animation/MeshDeformerInstance.h
@@ -49,10 +49,10 @@ public:
 
 
 	/** Called to allocate any persistent render resources */
-	ENGINE_API virtual void AllocateResources() PURE_VIRTUAL(, );
+	virtual void AllocateResources() PURE_VIRTUAL(, );
 
 	/** Called when persistent render resources should be released */
-	ENGINE_API virtual void ReleaseResources() PURE_VIRTUAL(, );
+	virtual void ReleaseResources() PURE_VIRTUAL(, );
 
 	/** Enumeration for workloads to EnqueueWork. */
 	enum EWorkLoad
@@ -84,18 +84,18 @@ public:
 	};
 
 	/** Enqueue the mesh deformer workload on a scene. */
-	ENGINE_API virtual void EnqueueWork(FEnqueueWorkDesc const& InDesc) PURE_VIRTUAL(, );
+	virtual void EnqueueWork(FEnqueueWorkDesc const& InDesc) PURE_VIRTUAL(, );
 	
 	/** Return the buffers that this deformer can potentially write to */
-	ENGINE_API virtual EMeshDeformerOutputBuffer GetOutputBuffers() const PURE_VIRTUAL(, return EMeshDeformerOutputBuffer::None; );
+	virtual EMeshDeformerOutputBuffer GetOutputBuffers() const PURE_VIRTUAL(, return EMeshDeformerOutputBuffer::None; );
 
 #if WITH_EDITORONLY_DATA
 	/** Reads back the deformed geometry and generates a mesh description */
-	ENGINE_API virtual bool RequestReadbackDeformerGeometry(TUniquePtr<FMeshDeformerGeometryReadbackRequest> InRequest) PURE_VIRTUAL(, return false; );
+	virtual bool RequestReadbackDeformerGeometry(TUniquePtr<FMeshDeformerGeometryReadbackRequest> InRequest) PURE_VIRTUAL(, return false; );
 #endif // WITH_EDITORONLY_DATA
 	
 	/** Returns the specific instance that directly represents the source deformer, this is needed as a deformer may create intermediate instances that aren't
 	 * necessarily user-facing.
 	 */
-	ENGINE_API virtual UMeshDeformerInstance* GetInstanceForSourceDeformer() PURE_VIRTUAL(, return this; ); 
+	virtual UMeshDeformerInstance* GetInstanceForSourceDeformer() PURE_VIRTUAL(, return this; ); 
 };
diff --git a/Engine/Source/Runtime/Engine/Classes/Components/LightComponent.h b/Engine/Source/Runtime/Engine/Classes/Components/LightComponent.h
--- a/Engine/Source/Runtime/Engine/Classes/Components/LightComponent.h
+++ b/Engine/Source/Runtime/Engine/Classes/Components/LightComponent.h
@@ -407,16 +407,16 @@ public:
 	/**
 	 * Return the homogenous position of the light.
 	 */
-	ENGINE_API virtual FVector4 GetLightPosition() const PURE_VIRTUAL(ULightComponent::GetPosition,return FVector4(););
+	virtual FVector4 GetLightPosition() const PURE_VIRTUAL(ULightComponent::GetPosition,return FVector4(););
 
 	/**
 	* @return ELightComponentType for the light component class
 	*/
-	ENGINE_API virtual ELightComponentType GetLightType() const PURE_VIRTUAL(ULightComponent::GetLightType,return LightType_MAX;);
+	virtual ELightComponentType GetLightType() const PURE_VIRTUAL(ULightComponent::GetLightType,return LightType_MAX;);
 
-	ENGINE_API virtual FLightmassLightSettings GetLightmassSettings() const PURE_VIRTUAL(ULightComponent::GetLightmassSettings,return FLightmassLightSettings(););
+	virtual FLightmassLightSettings GetLightmassSettings() const PURE_VIRTUAL(ULightComponent::GetLightmassSettings,return FLightmassLightSettings(););
 
-	ENGINE_API virtual float GetUniformPenumbraSize() const PURE_VIRTUAL(ULightComponent::GetUniformPenumbraSize,return 0;);
+	virtual float GetUniformPenumbraSize() const PURE_VIRTUAL(ULightComponent::GetUniformPenumbraSize,return 0;);
 
 	ENGINE_API virtual ELightUnits GetLightUnits() const;
 
diff --git a/Engine/Source/Runtime/Engine/Classes/Components/SkinnedMeshComponent.h b/Engine/Source/Runtime/Engine/Classes/Components/SkinnedMeshComponent.h
--- a/Engine/Source/Runtime/Engine/Classes/Components/SkinnedMeshComponent.h
+++ b/Engine/Source/Runtime/Engine/Classes/Components/SkinnedMeshComponent.h
@@ -1523,7 +1523,7 @@ public:
 	 * @param TickFunction Supplied as non null if we are running in a tick, allows us to create graph tasks for parallelism
 	 * 
 	 */
-	ENGINE_API virtual void RefreshBoneTransforms(FActorComponentTickFunction* TickFunction = nullptr) PURE_VIRTUAL(USkinnedMeshComponent::RefreshBoneTransforms, );
+	virtual void RefreshBoneTransforms(FActorComponentTickFunction* TickFunction = nullptr) PURE_VIRTUAL(USkinnedMeshComponent::RefreshBoneTransforms, );
 
 protected:
 	/** 
diff --git a/Engine/Source/Runtime/Engine/Classes/Components/SplineComponent.h b/Engine/Source/Runtime/Engine/Classes/Components/SplineComponent.h
--- a/Engine/Source/Runtime/Engine/Classes/Components/SplineComponent.h
+++ b/Engine/Source/Runtime/Engine/Classes/Components/SplineComponent.h
@@ -45,15 +45,15 @@ class USplineMetadata : public UObject
 
 public:
 	/** Insert point before index, lerping metadata between previous and next key values */
-	ENGINE_API virtual void InsertPoint(int32 Index, float t, bool bClosedLoop) PURE_VIRTUAL(USplineMetadata::InsertPoint, );
+	virtual void InsertPoint(int32 Index, float t, bool bClosedLoop) PURE_VIRTUAL(USplineMetadata::InsertPoint, );
 	/** Update point at index by lerping metadata between previous and next key values */
-	ENGINE_API virtual void UpdatePoint(int32 Index, float t, bool bClosedLoop) PURE_VIRTUAL(USplineMetadata::UpdatePoint, );
-	ENGINE_API virtual void AddPoint(float InputKey) PURE_VIRTUAL(USplineMetadata::AddPoint, );
-	ENGINE_API virtual void RemovePoint(int32 Index) PURE_VIRTUAL(USplineMetadata::RemovePoint, );
-	ENGINE_API virtual void DuplicatePoint(int32 Index) PURE_VIRTUAL(USplineMetadata::DuplicatePoint, );
-	ENGINE_API virtual void CopyPoint(const USplineMetadata* FromSplineMetadata, int32 FromIndex, int32 ToIndex) PURE_VIRTUAL(USplineMetadata::CopyPoint, );
-	ENGINE_API virtual void Reset(int32 NumPoints) PURE_VIRTUAL(USplineMetadata::Reset, );
-	ENGINE_API virtual void Fixup(int32 NumPoints, USplineComponent* SplineComp) PURE_VIRTUAL(USplineMetadata::Fixup, );
+	virtual void UpdatePoint(int32 Index, float t, bool bClosedLoop) PURE_VIRTUAL(USplineMetadata::UpdatePoint, );
+	virtual void AddPoint(float InputKey) PURE_VIRTUAL(USplineMetadata::AddPoint, );
+	virtual void RemovePoint(int32 Index) PURE_VIRTUAL(USplineMetadata::RemovePoint, );
+	virtual void DuplicatePoint(int32 Index) PURE_VIRTUAL(USplineMetadata::DuplicatePoint, );
+	virtual void CopyPoint(const USplineMetadata* FromSplineMetadata, int32 FromIndex, int32 ToIndex) PURE_VIRTUAL(USplineMetadata::CopyPoint, );
+	virtual void Reset(int32 NumPoints) PURE_VIRTUAL(USplineMetadata::Reset, );
+	virtual void Fixup(int32 NumPoints, USplineComponent* SplineComp) PURE_VIRTUAL(USplineMetadata::Fixup, );
 };
 
 USTRUCT()
diff --git a/Engine/Source/Runtime/Engine/Classes/Engine/Channel.h b/Engine/Source/Runtime/Engine/Classes/Engine/Channel.h
--- a/Engine/Source/Runtime/Engine/Classes/Engine/Channel.h
+++ b/Engine/Source/Runtime/Engine/Classes/Engine/Channel.h
@@ -114,7 +114,7 @@ public:
 	ENGINE_API virtual FString Describe();
 
 	/** Handle an incoming bunch. */
-	ENGINE_API virtual void ReceivedBunch( FInBunch& Bunch ) PURE_VIRTUAL(UChannel::ReceivedBunch,);
+	virtual void ReceivedBunch( FInBunch& Bunch ) PURE_VIRTUAL(UChannel::ReceivedBunch,);
 	
 	/** Positive acknowledgment processing. */
 	ENGINE_API virtual void ReceivedAck( int32 AckPacketId );
diff --git a/Engine/Source/Runtime/Engine/Classes/Engine/Engine.h b/Engine/Source/Runtime/Engine/Classes/Engine/Engine.h
--- a/Engine/Source/Runtime/Engine/Classes/Engine/Engine.h
+++ b/Engine/Source/Runtime/Engine/Classes/Engine/Engine.h
@@ -2424,7 +2424,7 @@ public:
 #endif // !UE_BUILD_SHIPPING
 
 	/** Update everything. */
-	ENGINE_API virtual void Tick( float DeltaSeconds, bool bIdleMode ) PURE_VIRTUAL(UEngine::Tick,);
+	virtual void Tick( float DeltaSeconds, bool bIdleMode ) PURE_VIRTUAL(UEngine::Tick,);
 
 	/**
 	 * Update FApp::CurrentTime / FApp::DeltaTime while taking into account max tick rate.
diff --git a/Engine/Source/Runtime/Engine/Classes/Engine/EngineBaseTypes.h b/Engine/Source/Runtime/Engine/Classes/Engine/EngineBaseTypes.h
--- a/Engine/Source/Runtime/Engine/Classes/Engine/EngineBaseTypes.h
+++ b/Engine/Source/Runtime/Engine/Classes/Engine/EngineBaseTypes.h
@@ -478,10 +478,10 @@ public:
 	 * @param CurrentThread - thread we are executing on, useful to pass along as new tasks are created
 	 * @param MyCompletionGraphEvent - completion event for this task. Useful for holding the completetion of this task until certain child tasks are complete.
 	 **/
-	ENGINE_API virtual void ExecuteTick(float DeltaTime, ELevelTick TickType, ENamedThreads::Type CurrentThread, const FGraphEventRef& MyCompletionGraphEvent) PURE_VIRTUAL(,);
+	virtual void ExecuteTick(float DeltaTime, ELevelTick TickType, ENamedThreads::Type CurrentThread, const FGraphEventRef& MyCompletionGraphEvent) PURE_VIRTUAL(,);
 	
 	/** Abstract function to describe this tick. Used to print messages about illegal cycles in the dependency graph **/
-	ENGINE_API virtual FString DiagnosticMessage() PURE_VIRTUAL(, return TEXT("DiagnosticMessage() not implemented"););
+	virtual FString DiagnosticMessage() PURE_VIRTUAL(, return TEXT("DiagnosticMessage() not implemented"););
 	
 	/** Function to give a 'context' for this tick, used for grouped active tick reporting */
 	virtual FName DiagnosticContext(bool bDetailed)
diff --git a/Engine/Source/Runtime/Engine/Classes/Engine/EngineCustomTimeStep.h b/Engine/Source/Runtime/Engine/Classes/Engine/EngineCustomTimeStep.h
--- a/Engine/Source/Runtime/Engine/Classes/Engine/EngineCustomTimeStep.h
+++ b/Engine/Source/Runtime/Engine/Classes/Engine/EngineCustomTimeStep.h
@@ -38,19 +38,19 @@ class UEngineCustomTimeStep : public UObject
 
 public:
 	/** This CustomTimeStep became the Engine's CustomTimeStep. */
-	ENGINE_API virtual bool Initialize(class UEngine* InEngine) PURE_VIRTUAL(UEngineCustomTimeStep::Initialize, return false;);
+	virtual bool Initialize(class UEngine* InEngine) PURE_VIRTUAL(UEngineCustomTimeStep::Initialize, return false;);
 
 	/** This CustomTimeStep stop being the Engine's CustomTimeStep. */
-	ENGINE_API virtual void Shutdown(class UEngine* InEngine) PURE_VIRTUAL(UEngineCustomTimeStep::Shutdown, );
+	virtual void Shutdown(class UEngine* InEngine) PURE_VIRTUAL(UEngineCustomTimeStep::Shutdown, );
 
 	/**
 	 * Update FApp::CurrentTime/FApp::DeltaTime and optionally wait until the end of the frame.
 	 * @return	true if the Engine's TimeStep should also be performed; false otherwise.
 	 */
-	ENGINE_API virtual bool UpdateTimeStep(class UEngine* InEngine) PURE_VIRTUAL(UEngineCustomTimeStep::UpdateTimeStep, return true;);
+	virtual bool UpdateTimeStep(class UEngine* InEngine) PURE_VIRTUAL(UEngineCustomTimeStep::UpdateTimeStep, return true;);
 
 	/** The state of the CustomTimeStep. */
-	ENGINE_API virtual ECustomTimeStepSynchronizationState GetSynchronizationState() const PURE_VIRTUAL(UEngineCustomTimeStep::GetSynchronizationState, return ECustomTimeStepSynchronizationState::Closed;);
+	virtual ECustomTimeStepSynchronizationState GetSynchronizationState() const PURE_VIRTUAL(UEngineCustomTimeStep::GetSynchronizationState, return ECustomTimeStepSynchronizationState::Closed;);
 
 	/** Get the display name of the custom time step. Allows child classes to provide more context to users. */
 	virtual FString GetDisplayName() const { return GetName(); };
diff --git a/Engine/Source/Runtime/Engine/Classes/Engine/NetConnection.h b/Engine/Source/Runtime/Engine/Classes/Engine/NetConnection.h
--- a/Engine/Source/Runtime/Engine/Classes/Engine/NetConnection.h
+++ b/Engine/Source/Runtime/Engine/Classes/Engine/NetConnection.h
@@ -1026,7 +1026,7 @@ public:
 	 * @param Traits		Special traits for the packet, passed down from the NetConnection through the PacketHandler
 	 */
 	// @todo: Traits should be passed within bit readers/writers, eventually
-	ENGINE_API virtual void LowLevelSend(void* Data, int32 CountBits, FOutPacketTraits& Traits)
+	virtual void LowLevelSend(void* Data, int32 CountBits, FOutPacketTraits& Traits)
 		PURE_VIRTUAL(UNetConnection::LowLevelSend,);
 
 	/** Validates the FBitWriter to make sure it's not in an error state */
@@ -1158,7 +1158,7 @@ public:
 	 * @param InMaxPacket the max packet size that will be used for sending
 	 * @param InPacketOverhead the packet overhead for this connection type
 	 */
-	ENGINE_API virtual void InitRemoteConnection(UNetDriver* InDriver, class FSocket* InSocket, const FURL& InURL, const class FInternetAddr& InRemoteAddr, EConnectionState InState, int32 InMaxPacket = 0, int32 InPacketOverhead = 0) PURE_VIRTUAL(UNetConnection::InitRemoteConnection, );
+	virtual void InitRemoteConnection(UNetDriver* InDriver, class FSocket* InSocket, const FURL& InURL, const class FInternetAddr& InRemoteAddr, EConnectionState InState, int32 InMaxPacket = 0, int32 InPacketOverhead = 0) PURE_VIRTUAL(UNetConnection::InitRemoteConnection, );
 	
 	/**
 	 * Initialize this connection instance *to* a remote source
@@ -1171,7 +1171,7 @@ public:
 	 * @param InMaxPacket the max packet size that will be used for sending
 	 * @param InPacketOverhead the packet overhead for this connection type
 	 */
-	ENGINE_API virtual void InitLocalConnection(UNetDriver* InDriver, class FSocket* InSocket, const FURL& InURL, EConnectionState InState, int32 InMaxPacket = 0, int32 InPacketOverhead = 0) PURE_VIRTUAL(UNetConnection::InitLocalConnection, );
+	virtual void InitLocalConnection(UNetDriver* InDriver, class FSocket* InSocket, const FURL& InURL, EConnectionState InState, int32 InMaxPacket = 0, int32 InPacketOverhead = 0) PURE_VIRTUAL(UNetConnection::InitLocalConnection, );
 	
 	/**
 	 * Initializes an "addressless" connection with the passed in settings
diff --git a/Engine/Source/Runtime/Engine/Classes/Engine/NetDriver.h b/Engine/Source/Runtime/Engine/Classes/Engine/NetDriver.h
--- a/Engine/Source/Runtime/Engine/Classes/Engine/NetDriver.h
+++ b/Engine/Source/Runtime/Engine/Classes/Engine/NetDriver.h
@@ -1484,7 +1484,7 @@ public:
 	 *
 	 * @return true if available, false otherwise
 	 */
-	ENGINE_API virtual bool IsAvailable() const PURE_VIRTUAL( UNetDriver::IsAvailable, return false;)
+	virtual bool IsAvailable() const PURE_VIRTUAL( UNetDriver::IsAvailable, return false;)
 
 	/**
 	 * Common initialization between server and client connection setup
@@ -1508,7 +1508,7 @@ public:
 	 * 
 	 * @return true if successful, false otherwise (check Error parameter)
 	 */
-	ENGINE_API virtual bool InitConnect(class FNetworkNotify* InNotify, const FURL& ConnectURL, FString& Error ) PURE_VIRTUAL( UNetDriver::InitConnect, return true;);
+	virtual bool InitConnect(class FNetworkNotify* InNotify, const FURL& ConnectURL, FString& Error ) PURE_VIRTUAL( UNetDriver::InitConnect, return true;);
 
 	/**
 	 * Initialize the network driver in server mode (listener)
@@ -1520,7 +1520,7 @@ public:
 	 *
 	 * @return true if successful, false otherwise (check Error parameter)
 	 */
-	ENGINE_API virtual bool InitListen(class FNetworkNotify* InNotify, FURL& ListenURL, bool bReuseAddressAndPort, FString& Error) PURE_VIRTUAL( UNetDriver::InitListen, return true;);
+	virtual bool InitListen(class FNetworkNotify* InNotify, FURL& ListenURL, bool bReuseAddressAndPort, FString& Error) PURE_VIRTUAL( UNetDriver::InitListen, return true;);
 
 	/** Initialize the list of destroyed net startup actors from the current World */
 	ENGINE_API virtual void InitDestroyedStartupActors();
@@ -1679,7 +1679,7 @@ public:
 	 * @param CountBits		The size of the packet data, in bits
 	 * @param Traits		Traits for the packet, if applicable
 	 */
-	ENGINE_API virtual void LowLevelSend(TSharedPtr<const FInternetAddr> Address, void* Data, int32 CountBits, FOutPacketTraits& Traits)
+	virtual void LowLevelSend(TSharedPtr<const FInternetAddr> Address, void* Data, int32 CountBits, FOutPacketTraits& Traits)
 		PURE_VIRTUAL(UNetDriver::LowLevelSend,);
 
 	/**
@@ -1814,7 +1814,7 @@ public:
 	/**
 	 * Get the socket subsytem appropriate for this net driver
 	 */
-	ENGINE_API virtual class ISocketSubsystem* GetSocketSubsystem() PURE_VIRTUAL(UNetDriver::GetSocketSubsystem, return NULL;);
+	virtual class ISocketSubsystem* GetSocketSubsystem() PURE_VIRTUAL(UNetDriver::GetSocketSubsystem, return NULL;);
 
 	/**
 	 * Associate a world with this net driver. 
@@ -1835,7 +1835,7 @@ public:
 	ENGINE_API virtual void ResetGameWorldState();
 
 	/** @return true if the net resource is valid or false if it should not be used */
-	ENGINE_API virtual bool IsNetResourceValid(void) PURE_VIRTUAL(UNetDriver::IsNetResourceValid, return false;);
+	virtual bool IsNetResourceValid(void) PURE_VIRTUAL(UNetDriver::IsNetResourceValid, return false;);
 
 	bool NetObjectIsDynamic(const UObject *Object) const;
 
