CameraCalibration: Expose render thread lens distortion lut generation function.

To be used by compositing to undistort footage before projecting it into the scene. For pre-processing, this needs to be called explicitely since the main render view LUT is not yet accessible.

diff --git a/Engine/Plugins/VirtualProduction/CameraCalibrationCore/Source/CameraCalibrationCore/Private/CameraCalibrationSubsystem.cpp b/Engine/Plugins/VirtualProduction/CameraCalibrationCore/Source/CameraCalibrationCore/Private/CameraCalibrationSubsystem.cpp
--- a/Engine/Plugins/VirtualProduction/CameraCalibrationCore/Source/CameraCalibrationCore/Private/CameraCalibrationSubsystem.cpp
+++ b/Engine/Plugins/VirtualProduction/CameraCalibrationCore/Source/CameraCalibrationCore/Private/CameraCalibrationSubsystem.cpp
@@ -186,6 +186,11 @@ void UCameraCalibrationSubsystem::ClearLensDistortionSVEState(ACameraActor* Came
 	SceneViewExtension->ClearDistortionState_AnyThread(CameraActor);
 }
 
+const TSharedPtr<FLensDistortionSceneViewExtension, ESPMode::ThreadSafe>& UCameraCalibrationSubsystem::GetSceneViewExtension() const
+{
+	return SceneViewExtension;
+}
+
 void UCameraCalibrationSubsystem::Initialize(FSubsystemCollectionBase& Collection)
 {
 	Super::Initialize(Collection);
diff --git a/Engine/Plugins/VirtualProduction/CameraCalibrationCore/Source/CameraCalibrationCore/Private/LensDistortionSceneViewExtension.cpp b/Engine/Plugins/VirtualProduction/CameraCalibrationCore/Source/CameraCalibrationCore/Private/LensDistortionSceneViewExtension.cpp
--- a/Engine/Plugins/VirtualProduction/CameraCalibrationCore/Source/CameraCalibrationCore/Private/LensDistortionSceneViewExtension.cpp
+++ b/Engine/Plugins/VirtualProduction/CameraCalibrationCore/Source/CameraCalibrationCore/Private/LensDistortionSceneViewExtension.cpp
@@ -231,7 +231,7 @@ public:
 
 IMPLEMENT_GLOBAL_SHADER(FRecenterSTDisplacementMapPS, "/Plugin/CameraCalibrationCore/Private/RecenterSTDisplacementMaps.usf", "MainPS", SF_Pixel);
 
-void FLensDistortionSceneViewExtension::DrawDisplacementMap_RenderThread(FRDGBuilder& GraphBuilder, const FLensDistortionState& CurrentState, EDistortionModel DistortionModel, float InverseOverscan, float CameraOverscan, const FVector2D& SensorSize, FRDGTextureRef& OutDistortionMapWithOverscan)
+void FLensDistortionSceneViewExtension::DrawDisplacementMap_RenderThread(FRDGBuilder& GraphBuilder, const FLensDistortionState& CurrentState, EDistortionModel DistortionModel, float InverseOverscan, float CameraOverscan, const FVector2D& SensorSize, FRDGTextureRef& OutDistortionMapWithOverscan) const
 {
 	if (CurrentState.DistortionInfo.Parameters.IsEmpty())
 	{
@@ -290,7 +290,7 @@ void FLensDistortionSceneViewExtension::DrawDisplacementMap_RenderThread(FRDGBui
 		FIntVector(FMath::DivideAndRoundUp(DistortionMapResolution.X, 8), FMath::DivideAndRoundUp(DistortionMapResolution.Y, 8), 1));
 }
 
-void FLensDistortionSceneViewExtension::CropDisplacementMap_RenderThread(FRDGBuilder& GraphBuilder, const FRDGTextureRef& InDistortionMapWithOverscan, FRDGTextureRef& OutDistortionMap)
+void FLensDistortionSceneViewExtension::CropDisplacementMap_RenderThread(FRDGBuilder& GraphBuilder, const FRDGTextureRef& InDistortionMapWithOverscan, FRDGTextureRef& OutDistortionMap) const
 {
 	FCropDistortionDisplacementMapCS::FParameters* PassParameters = GraphBuilder.AllocParameters<FCropDistortionDisplacementMapCS::FParameters>();
 
@@ -309,7 +309,7 @@ void FLensDistortionSceneViewExtension::CropDisplacementMap_RenderThread(FRDGBui
 		FIntVector(FMath::DivideAndRoundUp(LUTResolution.X, 8), FMath::DivideAndRoundUp(LUTResolution.Y, 8), 1));
 }
 
-void FLensDistortionSceneViewExtension::BlendDisplacementMaps_RenderThread(FRDGBuilder& GraphBuilder, const FDisplacementMapBlendingParams& BlendState, EDistortionModel DistortionModel, float InverseOverscan, float CameraOverscan, const FVector2D& SensorSize, FRDGTextureRef& OutDistortionMapWithOverscan)
+void FLensDistortionSceneViewExtension::BlendDisplacementMaps_RenderThread(FRDGBuilder& GraphBuilder, const FDisplacementMapBlendingParams& BlendState, EDistortionModel DistortionModel, float InverseOverscan, float CameraOverscan, const FVector2D& SensorSize, FRDGTextureRef& OutDistortionMapWithOverscan) const
 {
 	FBlendDistortionDisplacementMapCS::FParameters* PassParameters = GraphBuilder.AllocParameters<FBlendDistortionDisplacementMapCS::FParameters>();
 
@@ -392,7 +392,7 @@ void FLensDistortionSceneViewExtension::BlendDisplacementMaps_RenderThread(FRDGB
 		FIntVector(FMath::DivideAndRoundUp(DistortionMapResolution.X, 8), FMath::DivideAndRoundUp(DistortionMapResolution.Y, 8), 1));
 }
 
-void FLensDistortionSceneViewExtension::InvertDistortionMap_RenderThread(FRDGBuilder& GraphBuilder, const FRDGTextureRef& InDistortionMap, float InInverseOverscan, FRDGTextureRef& OutUndistortionMap)
+void FLensDistortionSceneViewExtension::InvertDistortionMap_RenderThread(FRDGBuilder& GraphBuilder, const FRDGTextureRef& InDistortionMap, float InInverseOverscan, FRDGTextureRef& OutUndistortionMap) const
 {
 	FInvertDisplacementParameters* PassParameters = GraphBuilder.AllocParameters<FInvertDisplacementParameters>();
 
@@ -482,7 +482,7 @@ void FLensDistortionSceneViewExtension::InvertDistortionMap_RenderThread(FRDGBui
 		});
 }
 
-void FLensDistortionSceneViewExtension::FillSTDisplacementMap_RenderThread(FRDGBuilder& GraphBuilder, const FRDGTextureRef& InUndisplacementMap, float InOverscan, FRDGTextureRef& OutFilledDisplacementMap)
+void FLensDistortionSceneViewExtension::FillSTDisplacementMap_RenderThread(FRDGBuilder& GraphBuilder, const FRDGTextureRef& InUndisplacementMap, float InOverscan, FRDGTextureRef& OutFilledDisplacementMap) const
 {
 	FFillSTDisplacementMapParameters* PassParameters = GraphBuilder.AllocParameters<FFillSTDisplacementMapParameters>();
 
@@ -569,7 +569,7 @@ void FLensDistortionSceneViewExtension::FillSTDisplacementMap_RenderThread(FRDGB
 		});
 }
 
-void FLensDistortionSceneViewExtension::RecenterSTDisplacementMap_RenderThread(FRDGBuilder& GraphBuilder, const FRDGTextureRef& InDisplacementMap, float InOverscan, FRDGTextureRef& OutRecenteredDisplacementMap)
+void FLensDistortionSceneViewExtension::RecenterSTDisplacementMap_RenderThread(FRDGBuilder& GraphBuilder, const FRDGTextureRef& InDisplacementMap, float InOverscan, FRDGTextureRef& OutRecenteredDisplacementMap) const
 {
 	FScreenPassRenderTarget Output;
 	Output.Texture = OutRecenteredDisplacementMap;
@@ -637,12 +637,22 @@ float FLensDistortionSceneViewExtension::GetInverseOverscan(ULensDistortionModel
 }
 
 void FLensDistortionSceneViewExtension::PreRenderView_RenderThread(FRDGBuilder& GraphBuilder, FSceneView& InView)
+{
+	FLensDistortionLUT ViewDistortionLUT;
+
+	if (RenderViewDistortionLUT(GraphBuilder, InView.ViewActor.ActorUniqueId, ViewDistortionLUT))
+	{
+		LensDistortion::SetLUTUnsafe(InView, ViewDistortionLUT);
+	}
+}
+
+bool FLensDistortionSceneViewExtension::RenderViewDistortionLUT(FRDGBuilder& GraphBuilder, uint32 InViewActorID, FLensDistortionLUT& ViewDistortionLUT) const
 {
 	FScopeLock ScopeLock(&DistortionStateMapCriticalSection);
-	FCameraDistortionProxy* Proxy = DistortionStateMap.Find(InView.ViewActor.ActorUniqueId);
+	const FCameraDistortionProxy* Proxy = DistortionStateMap.Find(InViewActorID);
 	if (!Proxy)
 	{
-		return;
+		return false;
 	}
 
 	FVector2D SensorSize = FVector2D(Proxy->FilmbackSettings.SensorWidth, Proxy->FilmbackSettings.SensorHeight);
@@ -670,8 +680,6 @@ void FLensDistortionSceneViewExtension::PreRenderView_RenderThread(FRDGBuilder&
 		}
 	}
 	
-	FLensDistortionLUT ViewDistortionLUT;
-	
 	if (DistortionModel == EDistortionModel::STMap)
 	{
 		// ST Maps generally have their top left corner corresponding to the distortion at the original frustum, not the overscanned frustum, which means that
@@ -714,8 +722,7 @@ void FLensDistortionSceneViewExtension::PreRenderView_RenderThread(FRDGBuilder&
 		ViewDistortionLUT.DistortionOverscan = Proxy->CameraOverscan * OverscanAdjustment;
 		ViewDistortionLUT.DistortionGridDimensions = FIntPoint(32 * Proxy->CameraOverscan, 20 * Proxy->CameraOverscan);
 		
-		LensDistortion::SetLUTUnsafe(InView, ViewDistortionLUT);
-		return;
+		return true;
 	}
 	
 	const bool bForwardDistort = IsDistortionModelForwardDistorting(DistortionModel);
@@ -768,5 +775,5 @@ void FLensDistortionSceneViewExtension::PreRenderView_RenderThread(FRDGBuilder&
 	InvertDistortionMap_RenderThread(GraphBuilder, DistortionMapWithOverscan, InverseOverscan, InverseDistortionMap);
 	CropDisplacementMap_RenderThread(GraphBuilder, DistortionMapWithOverscan, ForwardDistortionMap);
 
-	LensDistortion::SetLUTUnsafe(InView, ViewDistortionLUT);
+	return true;
 }
diff --git a/Engine/Plugins/VirtualProduction/CameraCalibrationCore/Source/CameraCalibrationCore/Public/CameraCalibrationSubsystem.h b/Engine/Plugins/VirtualProduction/CameraCalibrationCore/Source/CameraCalibrationCore/Public/CameraCalibrationSubsystem.h
--- a/Engine/Plugins/VirtualProduction/CameraCalibrationCore/Source/CameraCalibrationCore/Public/CameraCalibrationSubsystem.h
+++ b/Engine/Plugins/VirtualProduction/CameraCalibrationCore/Source/CameraCalibrationCore/Public/CameraCalibrationSubsystem.h
@@ -126,6 +126,9 @@ public:
 	/** Removes all distortion state from the Lens Distortion Scene View Extension for the input camera actor */
 	UE_API void ClearLensDistortionSVEState(ACameraActor* CameraActor);
 
+	/** Get lens distortion scene view extension. */
+	UE_API const TSharedPtr<FLensDistortionSceneViewExtension, ESPMode::ThreadSafe>& GetSceneViewExtension() const;
+
 private:
 	/** Default lens file to use when no override has been provided */
 	UPROPERTY(Transient)
diff --git a/Engine/Plugins/VirtualProduction/CameraCalibrationCore/Source/CameraCalibrationCore/Public/LensDistortionSceneViewExtension.h b/Engine/Plugins/VirtualProduction/CameraCalibrationCore/Source/CameraCalibrationCore/Public/LensDistortionSceneViewExtension.h
--- a/Engine/Plugins/VirtualProduction/CameraCalibrationCore/Source/CameraCalibrationCore/Public/LensDistortionSceneViewExtension.h
+++ b/Engine/Plugins/VirtualProduction/CameraCalibrationCore/Source/CameraCalibrationCore/Public/LensDistortionSceneViewExtension.h
@@ -6,6 +6,7 @@
 
 #include "LensFileRendering.h"
 
+struct FLensDistortionLUT;
 class UCameraComponent;
 
 /**
@@ -38,6 +39,9 @@ public:
 	virtual void PreRenderView_RenderThread(FRDGBuilder& GraphBuilder, FSceneView& InView) override;
 	//~End ISceneVIewExtension interface
 
+	/** Render forward & inverse displacement maps and return the engine lens distortion LUT. */
+	CAMERACALIBRATIONCORE_API bool RenderViewDistortionLUT(FRDGBuilder& GraphBuilder, uint32 InViewActorID, FLensDistortionLUT& ViewDistortionLUT) const;
+
 public:
 	/** Update the distortion state and blending params for the input camera */
 	void UpdateDistortionState_AnyThread(ACameraActor* CameraActor, FDisplacementMapBlendingParams DistortionState, ULensDistortionModelHandlerBase* LensDistortionHandler = nullptr);
@@ -47,22 +51,22 @@ public:
 
 private:
 	/** Use the input distortion state to draw a distortion displacement map */
-	void DrawDisplacementMap_RenderThread(FRDGBuilder& GraphBuilder, const FLensDistortionState& CurrentState, EDistortionModel DistortionModel, float InverseOverscan, float CameraOverscan, const FVector2D& SensorSize, FRDGTextureRef& OutDistortionMapWithOverscan);
+	void DrawDisplacementMap_RenderThread(FRDGBuilder& GraphBuilder, const FLensDistortionState& CurrentState, EDistortionModel DistortionModel, float InverseOverscan, float CameraOverscan, const FVector2D& SensorSize, FRDGTextureRef& OutDistortionMapWithOverscan) const;
 
 	/** Use the input blend parameters to draw multiple displacement maps and blend them together into a final distortion displacement map */
-	void BlendDisplacementMaps_RenderThread(FRDGBuilder& GraphBuilder, const FDisplacementMapBlendingParams& BlendState, EDistortionModel DistortionModel, float InverseOverscan, float CameraOverscan, const FVector2D& SensorSize, FRDGTextureRef& OutDistortionMapWithOverscan);
+	void BlendDisplacementMaps_RenderThread(FRDGBuilder& GraphBuilder, const FDisplacementMapBlendingParams& BlendState, EDistortionModel DistortionModel, float InverseOverscan, float CameraOverscan, const FVector2D& SensorSize, FRDGTextureRef& OutDistortionMapWithOverscan) const;
 
 	/** Crop the input overscanned distortion map to the original requested resolution */
-	void CropDisplacementMap_RenderThread(FRDGBuilder& GraphBuilder, const FRDGTextureRef& InDistortionMapWithOverscan, FRDGTextureRef& OutDistortionMap);
+	void CropDisplacementMap_RenderThread(FRDGBuilder& GraphBuilder, const FRDGTextureRef& InDistortionMapWithOverscan, FRDGTextureRef& OutDistortionMap) const;
 
 	/** Invert the input distortion map to generate a matching undistortion map (with no overscan) */
-	void InvertDistortionMap_RenderThread(FRDGBuilder& GraphBuilder, const FRDGTextureRef& InDistortionMap, float InInverseOverscan, FRDGTextureRef& OutUndistortionMap);
+	void InvertDistortionMap_RenderThread(FRDGBuilder& GraphBuilder, const FRDGTextureRef& InDistortionMap, float InInverseOverscan, FRDGTextureRef& OutUndistortionMap) const;
 
 	/** Performs a warp grid inverse of the specified distortion map to fill up its inverse to the specified overscan amount */
-	void FillSTDisplacementMap_RenderThread(FRDGBuilder& GraphBuilder, const FRDGTextureRef& InUndisplacementMap, float InOverscan, FRDGTextureRef& OutFilledDisplacementMap);
+	void FillSTDisplacementMap_RenderThread(FRDGBuilder& GraphBuilder, const FRDGTextureRef& InUndisplacementMap, float InOverscan, FRDGTextureRef& OutFilledDisplacementMap) const;
 
 	/** Recenters the specified distortion map within an overscanned map, filling the overscanned region with zeros */
-	void RecenterSTDisplacementMap_RenderThread(FRDGBuilder& GraphBuilder, const FRDGTextureRef& InDisplacementMap, float InOverscan, FRDGTextureRef& OutRecenteredDisplacementMap);
+	void RecenterSTDisplacementMap_RenderThread(FRDGBuilder& GraphBuilder, const FRDGTextureRef& InDisplacementMap, float InOverscan, FRDGTextureRef& OutRecenteredDisplacementMap) const;
 	
 	/** Gets whether the specified distortion model is forward distorting (generating the distort map first, then inverting for the undistort) or inverse distorting */
 	bool IsDistortionModelForwardDistorting(EDistortionModel InDistortionModel) const;
