LensDistortion: Minor function addition to check if the lens distortion scene view extension has any active distortion state.

diff --git a/Engine/Plugins/VirtualProduction/CameraCalibrationCore/Source/CameraCalibrationCore/Private/LensDistortionSceneViewExtension.cpp b/Engine/Plugins/VirtualProduction/CameraCalibrationCore/Source/CameraCalibrationCore/Private/LensDistortionSceneViewExtension.cpp
--- a/Engine/Plugins/VirtualProduction/CameraCalibrationCore/Source/CameraCalibrationCore/Private/LensDistortionSceneViewExtension.cpp
+++ b/Engine/Plugins/VirtualProduction/CameraCalibrationCore/Source/CameraCalibrationCore/Private/LensDistortionSceneViewExtension.cpp
@@ -50,6 +50,12 @@ void FLensDistortionSceneViewExtension::UpdateDistortionState_AnyThread(ACameraA
 	DistortionStateMap.Add(CameraActor->GetUniqueID(), MoveTemp(CamProxy));
 }
 
+bool FLensDistortionSceneViewExtension::HasDistortionState_AnyThread() const
+{
+	FScopeLock ScopeLock(&DistortionStateMapCriticalSection);
+	return DistortionStateMap.Num() > 0;
+}
+
 void FLensDistortionSceneViewExtension::ClearDistortionState_AnyThread(ACameraActor* CameraActor)
 {
 	FScopeLock ScopeLock(&DistortionStateMapCriticalSection);
diff --git a/Engine/Plugins/VirtualProduction/CameraCalibrationCore/Source/CameraCalibrationCore/Public/LensDistortionSceneViewExtension.h b/Engine/Plugins/VirtualProduction/CameraCalibrationCore/Source/CameraCalibrationCore/Public/LensDistortionSceneViewExtension.h
--- a/Engine/Plugins/VirtualProduction/CameraCalibrationCore/Source/CameraCalibrationCore/Public/LensDistortionSceneViewExtension.h
+++ b/Engine/Plugins/VirtualProduction/CameraCalibrationCore/Source/CameraCalibrationCore/Public/LensDistortionSceneViewExtension.h
@@ -46,6 +46,9 @@ public:
 	/** Update the distortion state and blending params for the input camera */
 	void UpdateDistortionState_AnyThread(ACameraActor* CameraActor, FDisplacementMapBlendingParams DistortionState, ULensDistortionModelHandlerBase* LensDistortionHandler = nullptr);
 
+	/** Returns true if any distortion state has been registered. */
+	CAMERACALIBRATIONCORE_API bool HasDistortionState_AnyThread() const;
+
 	/** Remove the distortion state and blending params for the input camera */
 	void ClearDistortionState_AnyThread(ACameraActor* CameraActor);
 
