CameraCalibration: Fix cx cy handling.

Materials were incorrect, one was hard-coded to 0.5 cx cy, and the other had an extra removal of (0.5-cx). Also the shader had 0.5 hard-coded. With these fixes the images generated by opentrackio now match, after converting the parameters to opencv units.

The solver still needs to be double-checked, as well as the anamorphic lens code.

[RN] Fixed center offset in the evaluation of the spherical lens model.
#virtualized

diff --git a/Engine/Plugins/VirtualProduction/CameraCalibrationCore/Shaders/Private/DrawDisplacementMaps.usf b/Engine/Plugins/VirtualProduction/CameraCalibrationCore/Shaders/Private/DrawDisplacementMaps.usf
--- a/Engine/Plugins/VirtualProduction/CameraCalibrationCore/Shaders/Private/DrawDisplacementMaps.usf
+++ b/Engine/Plugins/VirtualProduction/CameraCalibrationCore/Shaders/Private/DrawDisplacementMaps.usf
@@ -29,8 +29,8 @@ float P2;
 
 float2 SphericalDistort(float2 InputUV)
 {
-	float x_prime = (InputUV.x - 0.5) / FocalLength.x;
-	float y_prime = (InputUV.y - 0.5) / FocalLength.y;
+	float x_prime = (InputUV.x - ImageCenter.x) / FocalLength.x;
+	float y_prime = (InputUV.y - ImageCenter.y) / FocalLength.y;
 
 	float x_prime2 = x_prime * x_prime;
 	float y_prime2 = y_prime * y_prime;
