CameraCalibration: Allow each lens model to provide their own shader.

This makes it more feasible for third party lens distortion models, and was done in preparation to adding the models supported by OpenTrackIO.

#rnx

diff --git a/Engine/Plugins/VirtualProduction/CameraCalibrationCore/Shaders/Private/AnamorphicDistortion.usf b/Engine/Plugins/VirtualProduction/CameraCalibrationCore/Shaders/Private/AnamorphicDistortion.usf
new file mode 100644
--- /dev/null
+++ b/Engine/Plugins/VirtualProduction/CameraCalibrationCore/Shaders/Private/AnamorphicDistortion.usf
@@ -0,0 +1,122 @@
+// Copyright Epic Games, Inc. All Rights Reserved.
+
+/*=============================================================================
+	Implements the 3DE4 Anamorphic - Standard, Degree 4 model for distortion.
+=============================================================================*/
+
+#include "/Engine/Public/Platform.ush"
+#include "/Engine/Private/Common.ush"
+
+// Shared parameters
+float2 ThreadIdToUV;
+float InverseOverscan;
+float CameraOverscan;
+RWTexture2D<float2> OutDistortionMap;
+
+// Anamorphic distortion parameters (3DE4 Anamorphic - Standard, Degree 4 model)
+float PixelAspect;
+float2 FilmbackSize;
+float2 Squeeze;
+float LensRotation;
+float CX02;
+float CX04;
+float CX22;
+float CX24;
+float CX44;
+float CY02;
+float CY04;
+float CY22;
+float CY24;
+float CY44;
+
+float2 AnamorphicDistort(float2 InputUV)
+{
+	// The following implementation is based on the 3DE4 Anamorphic - Standard, Degree 4 model for distortion
+	// Pre-compute some values that will be useful later
+	float r_fb = (0.5) * sqrt(FilmbackSize.x*FilmbackSize.x + FilmbackSize.y*FilmbackSize.y);
+
+	float phi_mnt_rad = LensRotation * (PI / 180);
+	float cos_phi_mnt = cos(phi_mnt_rad);
+	float sin_phi_mnt = sin(phi_mnt_rad);
+
+	// Step 1: Transform input UVs into diagonally-normalized coordinates, based on the width and height of the filmback/image
+	float x_dn = (InputUV.x - 0.5) * (FilmbackSize.x / r_fb); // don't forget to factor in lens center offset later
+	float y_dn = (InputUV.y - 0.5) * (FilmbackSize.y / r_fb); // don't forget to factor in lens center offset later
+
+	// Step 2: Scale the x-coordinate by the inverse of the pixel aspect ratio
+	x_dn = x_dn / PixelAspect;
+
+	// Step 3: Rotate x and y by the inverse of the mounting angle
+	float x_dn_rotated = (x_dn * cos_phi_mnt) + (y_dn * sin_phi_mnt);
+	float y_dn_rotated = (y_dn * cos_phi_mnt) - (x_dn * sin_phi_mnt);
+
+	// Step 4: Compute the anamorphic distortion for x and y, to the 4th degree
+	float r_dn_sq = x_dn_rotated*x_dn_rotated + y_dn_rotated*y_dn_rotated;
+	float r_dn_4th = r_dn_sq * r_dn_sq;
+	float phi = atan2(y_dn_rotated, x_dn_rotated);
+
+	float x_dn_prime = x_dn_rotated * (1 + (CX02 * r_dn_sq) + (CX04 * r_dn_4th) + (CX22 * r_dn_sq * cos(2 * phi)) + (CX24 * r_dn_4th * cos(2 * phi)) + (CX44 * r_dn_4th * cos(4 * phi)));
+	float y_dn_prime = y_dn_rotated * (1 + (CY02 * r_dn_sq) + (CY04 * r_dn_4th) + (CY22 * r_dn_sq * cos(2 * phi)) + (CY24 * r_dn_4th * cos(2 * phi)) + (CY44 * r_dn_4th * cos(4 * phi)));
+
+	// Step 5: Scale the x-coordinate by the pixel aspect ratio, and the x- and y-coordinates by the lens-breathing squeeze factors
+	x_dn_prime = x_dn_prime * PixelAspect * Squeeze.x;
+	y_dn_prime = y_dn_prime * Squeeze.y;
+
+	// Step 6: Rotate x and y by the mounting angle
+	float x_dn_prime_rotated = (x_dn_prime * cos_phi_mnt) - (y_dn_prime * sin_phi_mnt);
+	float y_dn_prime_rotated = (y_dn_prime * cos_phi_mnt) + (x_dn_prime * sin_phi_mnt);
+
+	// Step 7: Transform diagonally normalized coordinates back into unit coordinates
+	float x_unit_prime = ((x_dn_prime_rotated * r_fb) / FilmbackSize.x) + 0.5;
+	float y_unit_prime = ((y_dn_prime_rotated * r_fb) / FilmbackSize.y) + 0.5;
+	
+	float2 output_uv;
+	output_uv.x = x_unit_prime;
+	output_uv.y = y_unit_prime;
+
+	return output_uv;
+}
+
+float2 Distort(float2 InputUV)
+{
+	return AnamorphicDistort(InputUV);
+}
+
+bool IsDistortionWellOriented(float2 InputUV, float2 DistortedUV, float2 UVOffset)
+{
+	float2 DistortedXOffset = Distort(InputUV + UVOffset * int2(1, 0));
+	float2 DistortedYOffset = Distort(InputUV + UVOffset * int2(0, 1));
+
+	float2 XOffset = DistortedXOffset - DistortedUV;
+	float2 YOffset = DistortedYOffset - DistortedUV;
+	
+	bool bIsFrontFacing = (XOffset.x * YOffset.y - XOffset.y * YOffset.x) >= 0;
+	bool bIsRightSideUp = (XOffset.x * UVOffset.x) >= 0 && (YOffset.y * UVOffset.y) >= 0;
+	return bIsFrontFacing && bIsRightSideUp;
+}
+
+float2 TryDistort(float2 InputUV)
+{
+	float2 DistortedPosition = Distort(InputUV);
+	float2 UVOffset = 0.5 * ThreadIdToUV;
+	bool bIsOrientedWell =
+		IsDistortionWellOriented(InputUV, DistortedPosition, UVOffset) &&
+		IsDistortionWellOriented(InputUV, DistortedPosition, -UVOffset);
+
+	return lerp(InputUV, DistortedPosition, bIsOrientedWell);
+}
+
+[numthreads(8, 8, 1)]
+void MainCS(uint3 DispatchThreadId : SV_DispatchThreadID)
+{
+	// Add half-pixel offset to ThreadID to get accurate UV position
+	float2 UndistortedUV = (DispatchThreadId.xy + 0.5) * ThreadIdToUV;
+	
+	// Rescale the UV by an overscan factor (>= 1.0) to widen the range of UVs to be distorted
+	UndistortedUV = ((UndistortedUV - 0.5) * InverseOverscan * CameraOverscan) + 0.5;
+	
+	float2 DistortedUV = TryDistort(UndistortedUV);
+	
+	// Compute the UV displacement and scale it by the CameraOverscan to account for the increased resolution of the overscanned scene texture 
+	OutDistortionMap[DispatchThreadId.xy] = (DistortedUV - UndistortedUV) / CameraOverscan;
+}
\ No newline at end of file
diff --git a/Engine/Plugins/VirtualProduction/CameraCalibrationCore/Shaders/Private/DrawDisplacementMaps.usf b/Engine/Plugins/VirtualProduction/CameraCalibrationCore/Shaders/Private/DrawDisplacementMaps.usf
deleted file mode 100644
diff --git a/Engine/Plugins/VirtualProduction/CameraCalibrationCore/Shaders/Private/SphericalDistortion.usf b/Engine/Plugins/VirtualProduction/CameraCalibrationCore/Shaders/Private/SphericalDistortion.usf
new file mode 100644
--- /dev/null
+++ b/Engine/Plugins/VirtualProduction/CameraCalibrationCore/Shaders/Private/SphericalDistortion.usf
@@ -0,0 +1,100 @@
+// Copyright Epic Games, Inc. All Rights Reserved.
+
+/*=============================================================================
+	Implements the Brown-Conrady model with 3 radial and 2 tangential components.
+=============================================================================*/
+
+#include "/Engine/Public/Platform.ush"
+#include "/Engine/Private/Common.ush"
+
+// Shared parameters
+float2 ThreadIdToUV;
+float InverseOverscan;
+float CameraOverscan;
+RWTexture2D<float2> OutDistortionMap;
+
+// Spherical distortion parameters (Brown-Conrady model)
+float2 FocalLength;
+float2 ImageCenter;
+float K1;
+float K2;
+float K3;
+float P1;
+float P2;
+
+float2 SphericalDistort(float2 InputUV)
+{
+	float x_prime = (InputUV.x - ImageCenter.x) / FocalLength.x;
+	float y_prime = (InputUV.y - ImageCenter.y) / FocalLength.y;
+
+	float x_prime2 = x_prime * x_prime;
+	float y_prime2 = y_prime * y_prime;
+
+	float r2 = x_prime2 + y_prime2;
+	float r4 = r2 * r2;
+	float r6 = r2 * r2 * r2;
+
+	// Compute radial distortion
+	float dr = 1 + K1 * r2 + K2 * r4 + K3 * r6;
+
+	// Compute tangential distortion
+	float dtx = 2 * P1 * x_prime * y_prime + P2 * (r2 + 2 * x_prime2);
+	float dty = P1 * (r2 + 2 * y_prime2) + 2 * P2 * x_prime * y_prime;
+
+	// Compute total distortion
+	float x_double_prime = (x_prime * dr) + dtx;
+	float y_double_prime = (y_prime * dr) + dty;
+
+	float x_undistort = (FocalLength.x * x_double_prime) + ImageCenter.x;
+	float y_undistort = (FocalLength.y * y_double_prime) + ImageCenter.y;
+
+	float2 result;
+	result.x = x_undistort;
+	result.y = y_undistort;
+
+	return result;
+}
+
+float2 Distort(float2 InputUV)
+{
+	return SphericalDistort(InputUV);
+}
+
+bool IsDistortionWellOriented(float2 InputUV, float2 DistortedUV, float2 UVOffset)
+{
+	float2 DistortedXOffset = Distort(InputUV + UVOffset * int2(1, 0));
+	float2 DistortedYOffset = Distort(InputUV + UVOffset * int2(0, 1));
+
+	float2 XOffset = DistortedXOffset - DistortedUV;
+	float2 YOffset = DistortedYOffset - DistortedUV;
+	
+	bool bIsFrontFacing = (XOffset.x * YOffset.y - XOffset.y * YOffset.x) >= 0;
+	bool bIsRightSideUp = (XOffset.x * UVOffset.x) >= 0 && (YOffset.y * UVOffset.y) >= 0;
+	return bIsFrontFacing && bIsRightSideUp;
+}
+
+float2 TryDistort(float2 InputUV)
+{
+	float2 DistortedPosition = Distort(InputUV);
+	float2 UVOffset = 0.5 * ThreadIdToUV;
+	bool bIsOrientedWell =
+		IsDistortionWellOriented(InputUV, DistortedPosition, UVOffset) &&
+		IsDistortionWellOriented(InputUV, DistortedPosition, -UVOffset);
+
+	return lerp(InputUV, DistortedPosition, bIsOrientedWell);
+}
+
+[numthreads(8, 8, 1)]
+void MainCS(uint3 DispatchThreadId : SV_DispatchThreadID)
+{
+	// Add half-pixel offset to ThreadID to get accurate UV position
+	float2 UndistortedUV = (DispatchThreadId.xy + 0.5) * ThreadIdToUV;
+	
+	// Rescale the UV by an overscan factor (>= 1.0) to widen the range of UVs to be distorted
+	UndistortedUV = ((UndistortedUV - 0.5) * InverseOverscan * CameraOverscan) + 0.5;
+	
+	float2 DistortedUV = TryDistort(UndistortedUV);
+	
+	// Compute the UV displacement and scale it by the CameraOverscan to account for the increased resolution of the overscanned scene texture 
+	OutDistortionMap[DispatchThreadId.xy] = (DistortedUV - UndistortedUV) / CameraOverscan;
+}
\ No newline at end of file
diff --git a/Engine/Plugins/VirtualProduction/CameraCalibrationCore/Source/CameraCalibrationCore/Private/AnamorphicLensDistortionModelHandler.cpp b/Engine/Plugins/VirtualProduction/CameraCalibrationCore/Source/CameraCalibrationCore/Private/AnamorphicLensDistortionModelHandler.cpp
--- a/Engine/Plugins/VirtualProduction/CameraCalibrationCore/Source/CameraCalibrationCore/Private/AnamorphicLensDistortionModelHandler.cpp
+++ b/Engine/Plugins/VirtualProduction/CameraCalibrationCore/Source/CameraCalibrationCore/Private/AnamorphicLensDistortionModelHandler.cpp
@@ -2,11 +2,62 @@
 
 #include "AnamorphicLensDistortionModelHandler.h"
 
+#include "CameraCalibrationCoreLog.h"
 #include "CameraCalibrationSettings.h"
 #include "Engine/TextureRenderTarget2D.h"
+#include "GlobalShader.h"
+#include "RenderGraphBuilder.h"
+#include "RenderGraphUtils.h"
+#include "ShaderParameterUtils.h"
+#include "ShaderCompilerCore.h"
 
 #include UE_INLINE_GENERATED_CPP_BY_NAME(AnamorphicLensDistortionModelHandler)
 
+// Anamorphic distortion shader parameter structure
+BEGIN_SHADER_PARAMETER_STRUCT(FAnamorphicDistortionParams, )
+	SHADER_PARAMETER(float, PixelAspect)
+	SHADER_PARAMETER(FVector2f, FilmbackSize)
+	SHADER_PARAMETER(FVector2f, Squeeze)
+	SHADER_PARAMETER(float, LensRotation)
+	SHADER_PARAMETER(float, CX02)
+	SHADER_PARAMETER(float, CX04)
+	SHADER_PARAMETER(float, CX22)
+	SHADER_PARAMETER(float, CX24)
+	SHADER_PARAMETER(float, CX44)
+	SHADER_PARAMETER(float, CY02)
+	SHADER_PARAMETER(float, CY04)
+	SHADER_PARAMETER(float, CY22)
+	SHADER_PARAMETER(float, CY24)
+	SHADER_PARAMETER(float, CY44)
+END_SHADER_PARAMETER_STRUCT()
+
+// Dedicated anamorphic distortion compute shader
+class FAnamorphicDistortionCS : public FGlobalShader
+{
+	DECLARE_GLOBAL_SHADER(FAnamorphicDistortionCS);
+	SHADER_USE_PARAMETER_STRUCT(FAnamorphicDistortionCS, FGlobalShader);
+
+	BEGIN_SHADER_PARAMETER_STRUCT(FParameters, )
+		SHADER_PARAMETER(FVector2f, ThreadIdToUV)
+		SHADER_PARAMETER(float, InverseOverscan)
+		SHADER_PARAMETER(float, CameraOverscan)
+		SHADER_PARAMETER_RDG_TEXTURE_UAV(RWTexture2D, OutDistortionMap)
+		SHADER_PARAMETER_STRUCT_INCLUDE(FAnamorphicDistortionParams, AnamorphicDistortionParams)
+	END_SHADER_PARAMETER_STRUCT()
+
+	static void ModifyCompilationEnvironment(const FGlobalShaderPermutationParameters& Parameters, FShaderCompilerEnvironment& OutEnvironment)
+	{
+		FGlobalShader::ModifyCompilationEnvironment(Parameters, OutEnvironment);
+	}
+
+	static bool ShouldCompilePermutation(const FGlobalShaderPermutationParameters& Parameters)
+	{
+		return IsFeatureLevelSupported(Parameters.Platform, ERHIFeatureLevel::SM5);
+	}
+};
+
+IMPLEMENT_GLOBAL_SHADER(FAnamorphicDistortionCS, "/Plugin/CameraCalibrationCore/Private/AnamorphicDistortion.usf", "MainCS", SF_Compute);
+
 void UAnamorphicLensDistortionModelHandler::InitializeHandler()
 {
 	LensModelClass = UAnamorphicLensModel::StaticClass();
@@ -185,3 +236,73 @@ void UAnamorphicLensDistortionModelHandler::InterpretDistortionParameters()
 {
 	LensModelClass->GetDefaultObject<ULensModel>()->FromArray<FAnamorphicDistortionParameters>(CurrentState.DistortionInfo.Parameters, AnamorphicParameters);
 }
+
+FString UAnamorphicLensDistortionModelHandler::GetDistortionShaderPath() const
+{
+	return TEXT("/Plugin/CameraCalibrationCore/Private/AnamorphicDistortion.usf");
+}
+
+
+void UAnamorphicLensDistortionModelHandler::ExecuteDistortionShader(FRDGBuilder& GraphBuilder, const FLensDistortionState& InCurrentState, float InverseOverscan, float CameraOverscan, const FVector2D& SensorSize, FRDGTextureRef& OutDistortionMap) const
+{
+	FIntPoint DistortionMapResolution = OutDistortionMap->Desc.Extent;
+
+	FAnamorphicDistortionCS::FParameters* ShaderParams = GraphBuilder.AllocParameters<FAnamorphicDistortionCS::FParameters>();
+	
+	ShaderParams->OutDistortionMap = GraphBuilder.CreateUAV(OutDistortionMap);
+	ShaderParams->ThreadIdToUV = FVector2f(1.0f) / FVector2f(DistortionMapResolution);
+	ShaderParams->InverseOverscan = InverseOverscan;
+	ShaderParams->CameraOverscan = CameraOverscan;
+
+	// Setup anamorphic distortion parameters
+	// Anamorphic model requires exactly 14 parameters (PixelAspect, CX02-CX44, CY02-CY44, SqueezeX, SqueezeY, LensRotation)
+	if (InCurrentState.DistortionInfo.Parameters.Num() == 14)
+	{
+		const float PixelAspect = InCurrentState.DistortionInfo.Parameters[0];
+		ShaderParams->AnamorphicDistortionParams.FilmbackSize = FVector2f(SensorSize.X * PixelAspect, SensorSize.Y);
+		ShaderParams->AnamorphicDistortionParams.PixelAspect = PixelAspect;
+		ShaderParams->AnamorphicDistortionParams.CX02 = InCurrentState.DistortionInfo.Parameters[1];
+		ShaderParams->AnamorphicDistortionParams.CX04 = InCurrentState.DistortionInfo.Parameters[2];
+		ShaderParams->AnamorphicDistortionParams.CX22 = InCurrentState.DistortionInfo.Parameters[3];
+		ShaderParams->AnamorphicDistortionParams.CX24 = InCurrentState.DistortionInfo.Parameters[4];
+		ShaderParams->AnamorphicDistortionParams.CX44 = InCurrentState.DistortionInfo.Parameters[5];
+		ShaderParams->AnamorphicDistortionParams.CY02 = InCurrentState.DistortionInfo.Parameters[6];
+		ShaderParams->AnamorphicDistortionParams.CY04 = InCurrentState.DistortionInfo.Parameters[7];
+		ShaderParams->AnamorphicDistortionParams.CY22 = InCurrentState.DistortionInfo.Parameters[8];
+		ShaderParams->AnamorphicDistortionParams.CY24 = InCurrentState.DistortionInfo.Parameters[9];
+		ShaderParams->AnamorphicDistortionParams.CY44 = InCurrentState.DistortionInfo.Parameters[10];
+		ShaderParams->AnamorphicDistortionParams.Squeeze = FVector2f(InCurrentState.DistortionInfo.Parameters[11], InCurrentState.DistortionInfo.Parameters[12]);
+		ShaderParams->AnamorphicDistortionParams.LensRotation = InCurrentState.DistortionInfo.Parameters[13];
+	}
+	else
+	{
+		UE_LOG(LogCameraCalibrationCore, Warning, 
+			TEXT("AnamorphicLensDistortionModelHandler: Invalid parameter count (%d), expected exactly 14. Using default values."), 
+			InCurrentState.DistortionInfo.Parameters.Num()
+		);
+
+		ShaderParams->AnamorphicDistortionParams.FilmbackSize = FVector2f(SensorSize.X, SensorSize.Y);
+		ShaderParams->AnamorphicDistortionParams.PixelAspect = 1.0f;
+		ShaderParams->AnamorphicDistortionParams.CX02 = 0.0f;
+		ShaderParams->AnamorphicDistortionParams.CX04 = 0.0f;
+		ShaderParams->AnamorphicDistortionParams.CX22 = 0.0f;
+		ShaderParams->AnamorphicDistortionParams.CX24 = 0.0f;
+		ShaderParams->AnamorphicDistortionParams.CX44 = 0.0f;
+		ShaderParams->AnamorphicDistortionParams.CY02 = 0.0f;
+		ShaderParams->AnamorphicDistortionParams.CY04 = 0.0f;
+		ShaderParams->AnamorphicDistortionParams.CY22 = 0.0f;
+		ShaderParams->AnamorphicDistortionParams.CY24 = 0.0f;
+		ShaderParams->AnamorphicDistortionParams.CY44 = 0.0f;
+		ShaderParams->AnamorphicDistortionParams.Squeeze = FVector2f(1.0f, 1.0f);
+		ShaderParams->AnamorphicDistortionParams.LensRotation = 0.0f;
+	}
+
+
+	TShaderMapRef<FAnamorphicDistortionCS> ComputeShader(GetGlobalShaderMap(GMaxRHIFeatureLevel));
+	FComputeShaderUtils::AddPass(
+		GraphBuilder,
+		RDG_EVENT_NAME("AnamorphicDistortionDisplacementMap"),
+		ComputeShader,
+		ShaderParams,
+		FIntVector(FMath::DivideAndRoundUp(DistortionMapResolution.X, 8), FMath::DivideAndRoundUp(DistortionMapResolution.Y, 8), 1));
+}
diff --git a/Engine/Plugins/VirtualProduction/CameraCalibrationCore/Source/CameraCalibrationCore/Private/LensDistortionSceneViewExtension.cpp b/Engine/Plugins/VirtualProduction/CameraCalibrationCore/Source/CameraCalibrationCore/Private/LensDistortionSceneViewExtension.cpp
--- a/Engine/Plugins/VirtualProduction/CameraCalibrationCore/Source/CameraCalibrationCore/Private/LensDistortionSceneViewExtension.cpp
+++ b/Engine/Plugins/VirtualProduction/CameraCalibrationCore/Source/CameraCalibrationCore/Private/LensDistortionSceneViewExtension.cpp
@@ -14,8 +14,6 @@
 #include "ShaderParameterStruct.h"
 #include "SystemTextures.h"
 #include "TextureResource.h"
-#include "Models/AnamorphicLensModel.h"
-#include "Models/SphericalLensModel.h"
 
 TAutoConsoleVariable<int32> CVarLensDistortionInvertGridDensity(
 	TEXT("r.LensDistortion.InvertGridDensity"),
@@ -62,59 +60,6 @@ void FLensDistortionSceneViewExtension::ClearDistortionState_AnyThread(ACameraAc
 	DistortionStateMap.Remove(CameraActor->GetUniqueID());
 }
 
-BEGIN_SHADER_PARAMETER_STRUCT(FSphericalDistortionParams, )
-	SHADER_PARAMETER(FVector2f, FocalLength)
-	SHADER_PARAMETER(FVector2f, ImageCenter)
-	SHADER_PARAMETER(float, K1)
-	SHADER_PARAMETER(float, K2)
-	SHADER_PARAMETER(float, K3)
-	SHADER_PARAMETER(float, P1)
-	SHADER_PARAMETER(float, P2)
-END_SHADER_PARAMETER_STRUCT()
-
-BEGIN_SHADER_PARAMETER_STRUCT(FAnamorphicDistortionParams, )
-	SHADER_PARAMETER(float, PixelAspect)
-	SHADER_PARAMETER(FVector2f, FilmbackSize)
-	SHADER_PARAMETER(FVector2f, Squeeze)
-	SHADER_PARAMETER(float, LensRotation)
-	SHADER_PARAMETER(float, CX02)
-	SHADER_PARAMETER(float, CX04)
-	SHADER_PARAMETER(float, CX22)
-	SHADER_PARAMETER(float, CX24)
-	SHADER_PARAMETER(float, CX44)
-	SHADER_PARAMETER(float, CY02)
-	SHADER_PARAMETER(float, CY04)
-	SHADER_PARAMETER(float, CY22)
-	SHADER_PARAMETER(float, CY24)
-	SHADER_PARAMETER(float, CY44)
-END_SHADER_PARAMETER_STRUCT()
-
-class FDrawDistortionDisplacementMapCS : public FGlobalShader
-{
-	DECLARE_GLOBAL_SHADER(FDrawDistortionDisplacementMapCS);
-	SHADER_USE_PARAMETER_STRUCT(FDrawDistortionDisplacementMapCS, FGlobalShader);
-
-	class FDistortionModelDim : SHADER_PERMUTATION_ENUM_CLASS("DISTORTION_MODEL", FLensDistortionSceneViewExtension::EDistortionModel);
-	using FPermutationDomain = TShaderPermutationDomain<FDistortionModelDim>;
-	
-	BEGIN_SHADER_PARAMETER_STRUCT(FParameters, )
-		SHADER_PARAMETER(FVector2f, ThreadIdToUV)
-		SHADER_PARAMETER(float, InverseOverscan)
-		SHADER_PARAMETER(float, CameraOverscan)
-		SHADER_PARAMETER_STRUCT_INCLUDE(FSphericalDistortionParams, SphericalDistortionParams)
-		SHADER_PARAMETER_STRUCT_INCLUDE(FAnamorphicDistortionParams, AnamorphicDistortionParams)
-		SHADER_PARAMETER_RDG_TEXTURE_UAV(RWTexture2D, OutDistortionMap)
-	END_SHADER_PARAMETER_STRUCT()
-
-	// Called by the engine to determine which permutations to compile for this shader
-	static bool ShouldCompilePermutation(const FGlobalShaderPermutationParameters& Parameters)
-	{
-		return IsFeatureLevelSupported(Parameters.Platform, ERHIFeatureLevel::SM5);
-	}
-};
-
-IMPLEMENT_GLOBAL_SHADER(FDrawDistortionDisplacementMapCS, "/Plugin/CameraCalibrationCore/Private/DrawDisplacementMaps.usf", "MainCS", SF_Compute);
-
 
 class FBlendDistortionDisplacementMapCS : public FGlobalShader
 {
@@ -237,7 +182,7 @@ public:
 
 IMPLEMENT_GLOBAL_SHADER(FRecenterSTDisplacementMapPS, "/Plugin/CameraCalibrationCore/Private/RecenterSTDisplacementMaps.usf", "MainPS", SF_Pixel);
 
-void FLensDistortionSceneViewExtension::DrawDisplacementMap_RenderThread(FRDGBuilder& GraphBuilder, const FLensDistortionState& CurrentState, EDistortionModel DistortionModel, float InverseOverscan, float CameraOverscan, const FVector2D& SensorSize, FRDGTextureRef& OutDistortionMapWithOverscan) const
+void FLensDistortionSceneViewExtension::DrawDisplacementMap_RenderThread(FRDGBuilder& GraphBuilder, const FLensDistortionState& CurrentState, ULensDistortionModelHandlerBase* ModelHandler, float InverseOverscan, float CameraOverscan, const FVector2D& SensorSize, FRDGTextureRef& OutDistortionMapWithOverscan) const
 {
 	if (CurrentState.DistortionInfo.Parameters.IsEmpty())
 	{
@@ -245,55 +190,18 @@ void FLensDistortionSceneViewExtension::DrawDisplacementMap_RenderThread(FRDGBui
 		return;
 	}
 
-	FDrawDistortionDisplacementMapCS::FPermutationDomain PermutationVector;
-	PermutationVector.Set<FDrawDistortionDisplacementMapCS::FDistortionModelDim>(DistortionModel);
-	
-	FDrawDistortionDisplacementMapCS::FParameters* PassParameters = GraphBuilder.AllocParameters<FDrawDistortionDisplacementMapCS::FParameters>();
-
-	PassParameters->OutDistortionMap = GraphBuilder.CreateUAV(OutDistortionMapWithOverscan);
-
-	FIntPoint DistortionMapResolution = OutDistortionMapWithOverscan->Desc.Extent;
-	PassParameters->ThreadIdToUV = FVector2f(1.0f) / FVector2f(DistortionMapResolution);
-
-	if (DistortionModel == EDistortionModel::SphericalDistortion)
-	{
-		PassParameters->SphericalDistortionParams.ImageCenter = FVector2f(CurrentState.ImageCenter.PrincipalPoint);
-		PassParameters->SphericalDistortionParams.FocalLength = FVector2f(CurrentState.FocalLengthInfo.FxFy);
-
-		PassParameters->SphericalDistortionParams.K1 = CurrentState.DistortionInfo.Parameters[0];
-		PassParameters->SphericalDistortionParams.K2 = CurrentState.DistortionInfo.Parameters[1];
-		PassParameters->SphericalDistortionParams.K3 = CurrentState.DistortionInfo.Parameters[2];
-		PassParameters->SphericalDistortionParams.P1 = CurrentState.DistortionInfo.Parameters[3];
-		PassParameters->SphericalDistortionParams.P2 = CurrentState.DistortionInfo.Parameters[4];
-	}
-	else if (DistortionModel == EDistortionModel::AnamorphicDistortion)
+	// Check if we have a valid handler
+	if (!ModelHandler)
 	{
-		PassParameters->AnamorphicDistortionParams.FilmbackSize = FVector2f(SensorSize.X * CurrentState.DistortionInfo.Parameters[0], SensorSize.Y);
-		PassParameters->AnamorphicDistortionParams.PixelAspect = CurrentState.DistortionInfo.Parameters[0];
-		PassParameters->AnamorphicDistortionParams.CX02 = CurrentState.DistortionInfo.Parameters[1];
-		PassParameters->AnamorphicDistortionParams.CX04 = CurrentState.DistortionInfo.Parameters[2];
-		PassParameters->AnamorphicDistortionParams.CX22 = CurrentState.DistortionInfo.Parameters[3];
-		PassParameters->AnamorphicDistortionParams.CX24 = CurrentState.DistortionInfo.Parameters[4];
-		PassParameters->AnamorphicDistortionParams.CX44 = CurrentState.DistortionInfo.Parameters[5];
-		PassParameters->AnamorphicDistortionParams.CY02 = CurrentState.DistortionInfo.Parameters[6];
-		PassParameters->AnamorphicDistortionParams.CY04 = CurrentState.DistortionInfo.Parameters[7];
-		PassParameters->AnamorphicDistortionParams.CY22 = CurrentState.DistortionInfo.Parameters[8];
-		PassParameters->AnamorphicDistortionParams.CY24 = CurrentState.DistortionInfo.Parameters[9];
-		PassParameters->AnamorphicDistortionParams.CY44 = CurrentState.DistortionInfo.Parameters[10];
-		PassParameters->AnamorphicDistortionParams.Squeeze = FVector2f(CurrentState.DistortionInfo.Parameters[11], CurrentState.DistortionInfo.Parameters[12]);
-		PassParameters->AnamorphicDistortionParams.LensRotation = CurrentState.DistortionInfo.Parameters[13];
+		UE_LOG(LogTemp, Warning, TEXT("No model handler available in DrawDisplacementMap_RenderThread, using dummy texture"));
+		OutDistortionMapWithOverscan = GSystemTextures.GetBlackDummy(GraphBuilder);
+		return;
 	}
-	
-	PassParameters->InverseOverscan = InverseOverscan;
-	PassParameters->CameraOverscan = CameraOverscan;
 
-	TShaderMapRef<FDrawDistortionDisplacementMapCS> ComputeShader(GetGlobalShaderMap(GMaxRHIFeatureLevel), PermutationVector);
-	FComputeShaderUtils::AddPass(
-		GraphBuilder,
-		RDG_EVENT_NAME("DrawDistortionDisplacementMap"),
-		ComputeShader,
-		PassParameters,
-		FIntVector(FMath::DivideAndRoundUp(DistortionMapResolution.X, 8), FMath::DivideAndRoundUp(DistortionMapResolution.Y, 8), 1));
+	FIntPoint DistortionMapResolution = OutDistortionMapWithOverscan->Desc.Extent;
+
+	// Fully generic shader execution through handler - enables unlimited third-party extensibility
+	ModelHandler->ExecuteDistortionShader(GraphBuilder, CurrentState, InverseOverscan, CameraOverscan, SensorSize, OutDistortionMapWithOverscan);
 }
 
 void FLensDistortionSceneViewExtension::CropDisplacementMap_RenderThread(FRDGBuilder& GraphBuilder, const FRDGTextureRef& InDistortionMapWithOverscan, FRDGTextureRef& OutDistortionMap) const
@@ -315,14 +223,14 @@ void FLensDistortionSceneViewExtension::CropDisplacementMap_RenderThread(FRDGBui
 		FIntVector(FMath::DivideAndRoundUp(LUTResolution.X, 8), FMath::DivideAndRoundUp(LUTResolution.Y, 8), 1));
 }
 
-void FLensDistortionSceneViewExtension::BlendDisplacementMaps_RenderThread(FRDGBuilder& GraphBuilder, const FDisplacementMapBlendingParams& BlendState, EDistortionModel DistortionModel, float InverseOverscan, float CameraOverscan, const FVector2D& SensorSize, FRDGTextureRef& OutDistortionMapWithOverscan) const
+void FLensDistortionSceneViewExtension::BlendDisplacementMaps_RenderThread(FRDGBuilder& GraphBuilder, const FDisplacementMapBlendingParams& BlendState, ULensDistortionModelHandlerBase* ModelHandler, float InverseOverscan, float CameraOverscan, const FVector2D& SensorSize, FRDGTextureRef& OutDistortionMapWithOverscan) const
 {
 	FBlendDistortionDisplacementMapCS::FParameters* PassParameters = GraphBuilder.AllocParameters<FBlendDistortionDisplacementMapCS::FParameters>();
 
 	// Draw the first distortion map, which should always be valid
 	{
 		FRDGTextureRef Distortion1 = GraphBuilder.CreateTexture(OutDistortionMapWithOverscan->Desc, TEXT("DistortingDisplacement1"));
-		DrawDisplacementMap_RenderThread(GraphBuilder, BlendState.States[0], DistortionModel, InverseOverscan, CameraOverscan, SensorSize, Distortion1);
+		DrawDisplacementMap_RenderThread(GraphBuilder, BlendState.States[0], ModelHandler, InverseOverscan, CameraOverscan, SensorSize, Distortion1);
 		PassParameters->InputDistortionMap1 = GraphBuilder.CreateSRV(Distortion1);
 	}
 
@@ -330,7 +238,7 @@ void FLensDistortionSceneViewExtension::BlendDisplacementMaps_RenderThread(FRDGB
 	if (BlendState.BlendType != EDisplacementMapBlendType::OneFocusOneZoom)
 	{
 		FRDGTextureRef Distortion2 = GraphBuilder.CreateTexture(OutDistortionMapWithOverscan->Desc, TEXT("DistortingDisplacement2"));
-		DrawDisplacementMap_RenderThread(GraphBuilder, BlendState.States[1], DistortionModel, InverseOverscan, CameraOverscan, SensorSize, Distortion2);
+		DrawDisplacementMap_RenderThread(GraphBuilder, BlendState.States[1], ModelHandler, InverseOverscan, CameraOverscan, SensorSize, Distortion2);
 		PassParameters->InputDistortionMap2 = GraphBuilder.CreateSRV(Distortion2);
 	}
 
@@ -340,8 +248,8 @@ void FLensDistortionSceneViewExtension::BlendDisplacementMaps_RenderThread(FRDGB
 		FRDGTextureRef Distortion3 = GraphBuilder.CreateTexture(OutDistortionMapWithOverscan->Desc, TEXT("DistortingDisplacement3"));
 		FRDGTextureRef Distortion4 = GraphBuilder.CreateTexture(OutDistortionMapWithOverscan->Desc, TEXT("DistortingDisplacement4"));
 
-		DrawDisplacementMap_RenderThread(GraphBuilder, BlendState.States[2], DistortionModel, InverseOverscan, CameraOverscan, SensorSize, Distortion3);
-		DrawDisplacementMap_RenderThread(GraphBuilder, BlendState.States[3], DistortionModel, InverseOverscan, CameraOverscan, SensorSize, Distortion4);
+		DrawDisplacementMap_RenderThread(GraphBuilder, BlendState.States[2], ModelHandler, InverseOverscan, CameraOverscan, SensorSize, Distortion3);
+		DrawDisplacementMap_RenderThread(GraphBuilder, BlendState.States[3], ModelHandler, InverseOverscan, CameraOverscan, SensorSize, Distortion4);
 
 		PassParameters->InputDistortionMap3 = GraphBuilder.CreateSRV(Distortion3);
 		PassParameters->InputDistortionMap4 = GraphBuilder.CreateSRV(Distortion4);
@@ -617,30 +525,6 @@ void FLensDistortionSceneViewExtension::RecenterSTDisplacementMap_RenderThread(F
 	});
 }
 
-bool FLensDistortionSceneViewExtension::IsDistortionModelForwardDistorting(EDistortionModel InDistortionModel) const
-{
-	switch (InDistortionModel)
-	{
-	case EDistortionModel::AnamorphicDistortion:
-		return false;
-
-	case EDistortionModel::SphericalDistortion:
-	default:
-		return true;
-	}
-}
-
-float FLensDistortionSceneViewExtension::GetInverseOverscan(ULensDistortionModelHandlerBase* LensDistortionHandler, EDistortionModel InDistortionModel) const
-{
-	if (IsDistortionModelForwardDistorting(InDistortionModel))
-	{
-		return LensDistortionHandler->ComputeOverscanFactor();
-	}
-	else
-	{
-		return LensDistortionHandler->ComputeInverseOverscanFactor();
-	}
-}
 
 void FLensDistortionSceneViewExtension::PreRenderView_RenderThread(FRDGBuilder& GraphBuilder, FSceneView& InView)
 {
@@ -666,27 +550,16 @@ bool FLensDistortionSceneViewExtension::RenderViewDistortionLUT(FRDGBuilder& Gra
 
 	TStrongObjectPtr<ULensDistortionModelHandlerBase> ModelHandler = Proxy->LensDistortionHandler.Pin();
 	
-	EDistortionModel DistortionModel = EDistortionModel::None;
+	bool bIsSTMap = false;
 	if (ModelHandler.IsValid())
 	{
-		if (Proxy->Params.States[0].DistortionInfo.Parameters.Num())
+		if (!Proxy->Params.States[0].DistortionInfo.Parameters.Num())
 		{
-			if (ModelHandler->GetLensModelClass() == USphericalLensModel::StaticClass())
-			{
-				DistortionModel = EDistortionModel::SphericalDistortion;
-			}
-			else if (ModelHandler->GetLensModelClass() == UAnamorphicLensModel::StaticClass())
-			{
-				DistortionModel = EDistortionModel::AnamorphicDistortion;
-			}
-		}
-		else
-		{
-			DistortionModel = EDistortionModel::STMap;
+			bIsSTMap = true;
 		}
 	}
 	
-	if (DistortionModel == EDistortionModel::STMap)
+	if (bIsSTMap)
 	{
 		// ST Maps generally have their top left corner corresponding to the distortion at the original frustum, not the overscanned frustum, which means that
 		// for cases where the distortion would require overscanned pixels, we must generate those distortion values for the distortion map from the undistortion map
@@ -731,14 +604,12 @@ bool FLensDistortionSceneViewExtension::RenderViewDistortionLUT(FRDGBuilder& Gra
 		return true;
 	}
 	
-	const bool bForwardDistort = IsDistortionModelForwardDistorting(DistortionModel);
-	
 	ScopeLock.Unlock();
 
-
 	// Create the distortion map and undistortion map textures for the FLensDistortionLUT for this frame
 	const FIntPoint DisplacementMapResolution = FIntPoint(256, 256);
 
+	// Define forward and reverse distortion texture descriptors (matching original logic)
 	FRDGTextureDesc ForwardDistortionMapDesc = FRDGTextureDesc::Create2D(
 		DisplacementMapResolution,
 		PF_G32R32F,
@@ -750,15 +621,18 @@ bool FLensDistortionSceneViewExtension::RenderViewDistortionLUT(FRDGBuilder& Gra
 		PF_G32R32F,
 		FClearValueBinding::Black,
 		TexCreate_ShaderResource | TexCreate_RenderTargetable);
-	
+		
+	// Choose texture descriptors based on distortion direction (matching original logic)
+	const bool bForwardDistort = ModelHandler->IsForwardDistorting();
 	FRDGTextureDesc DistortionMapDesc = bForwardDistort ? ForwardDistortionMapDesc : InverseDistortionMapDesc;
 	FRDGTextureDesc UndistortionMapDesc = bForwardDistort ? InverseDistortionMapDesc : ForwardDistortionMapDesc;
 	
 	ViewDistortionLUT.DistortingDisplacementTexture = GraphBuilder.CreateTexture(DistortionMapDesc, TEXT("DistortionDisplacementMap"));
 	ViewDistortionLUT.UndistortingDisplacementTexture = GraphBuilder.CreateTexture(UndistortionMapDesc, TEXT("UndistortionDisplacementMap"));
 
-	// In order to guarantee that we can generate a complete undistortion map, the distortion map we invert needs to have some overscan 
-	float InverseOverscan = FMath::Clamp(GetInverseOverscan(ModelHandler.Get(), DistortionModel), 0.0f, 2.0f);
+	// Let each handler compute its own overscan - handlers know their own mathematical properties
+	float RawInverseOverscan = ModelHandler->IsForwardDistorting() ? ModelHandler->ComputeOverscanFactor() : ModelHandler->ComputeInverseOverscanFactor();
+	float InverseOverscan = FMath::Clamp(RawInverseOverscan, 0.0f, 2.0f);
 	
 	// Adjust the overscan resolution to be square, with each side being a multiple of 8
 	const FIntPoint OverscanResolution = FIntPoint(FMath::CeilToInt(InverseOverscan * 32) * 8);
@@ -773,11 +647,12 @@ bool FLensDistortionSceneViewExtension::RenderViewDistortionLUT(FRDGBuilder& Gra
 
 	FRDGTextureRef DistortionMapWithOverscan = GraphBuilder.CreateTexture(OverscanDesc, TEXT("DistortionMapWithOverscan"));
 
+	// Generate distortion map through handler, then invert it to create undistortion map
 	// Anamorphic model uses a precise equation for its _undistort_ map instead of its distort map, so draw the undistort map first and invert it to find the distort map
 	FRDGTextureRef ForwardDistortionMap = bForwardDistort ? ViewDistortionLUT.DistortingDisplacementTexture : ViewDistortionLUT.UndistortingDisplacementTexture;
 	FRDGTextureRef InverseDistortionMap = bForwardDistort ? ViewDistortionLUT.UndistortingDisplacementTexture : ViewDistortionLUT.DistortingDisplacementTexture;
 	
-	BlendDisplacementMaps_RenderThread(GraphBuilder, BlendState, DistortionModel, InverseOverscan, Proxy->CameraOverscan, SensorSize, DistortionMapWithOverscan);
+	BlendDisplacementMaps_RenderThread(GraphBuilder, BlendState, ModelHandler.Get(), InverseOverscan, Proxy->CameraOverscan, SensorSize, DistortionMapWithOverscan);
 	InvertDistortionMap_RenderThread(GraphBuilder, DistortionMapWithOverscan, InverseOverscan, InverseDistortionMap);
 	CropDisplacementMap_RenderThread(GraphBuilder, DistortionMapWithOverscan, ForwardDistortionMap);
 
diff --git a/Engine/Plugins/VirtualProduction/CameraCalibrationCore/Source/CameraCalibrationCore/Private/SphericalLensDistortionModelHandler.cpp b/Engine/Plugins/VirtualProduction/CameraCalibrationCore/Source/CameraCalibrationCore/Private/SphericalLensDistortionModelHandler.cpp
--- a/Engine/Plugins/VirtualProduction/CameraCalibrationCore/Source/CameraCalibrationCore/Private/SphericalLensDistortionModelHandler.cpp
+++ b/Engine/Plugins/VirtualProduction/CameraCalibrationCore/Source/CameraCalibrationCore/Private/SphericalLensDistortionModelHandler.cpp
@@ -2,12 +2,51 @@
 
 #include "SphericalLensDistortionModelHandler.h"
 
-#include "Engine/TextureRenderTarget2D.h"
+#include "CameraCalibrationCoreLog.h"
 #include "CameraCalibrationSettings.h"
+#include "Engine/TextureRenderTarget2D.h"
+#include "GlobalShader.h"
 #include "Math/NumericLimits.h"
+#include "RenderGraphBuilder.h"
+#include "RenderGraphUtils.h"
+#include "ShaderParameterUtils.h"
+#include "ShaderCompilerCore.h"
 
 #include UE_INLINE_GENERATED_CPP_BY_NAME(SphericalLensDistortionModelHandler)
 
+// Spherical distortion shader parameter structure
+BEGIN_SHADER_PARAMETER_STRUCT(FSphericalDistortionParams, )
+	SHADER_PARAMETER(FVector2f, FocalLength)
+	SHADER_PARAMETER(FVector2f, ImageCenter)
+	SHADER_PARAMETER(float, K1)
+	SHADER_PARAMETER(float, K2)
+	SHADER_PARAMETER(float, K3)
+	SHADER_PARAMETER(float, P1)
+	SHADER_PARAMETER(float, P2)
+END_SHADER_PARAMETER_STRUCT()
+
+// Dedicated spherical distortion compute shader
+class FSphericalDistortionCS : public FGlobalShader
+{
+	DECLARE_GLOBAL_SHADER(FSphericalDistortionCS);
+	SHADER_USE_PARAMETER_STRUCT(FSphericalDistortionCS, FGlobalShader);
+
+	BEGIN_SHADER_PARAMETER_STRUCT(FParameters, )
+		SHADER_PARAMETER(FVector2f, ThreadIdToUV)
+		SHADER_PARAMETER(float, InverseOverscan)
+		SHADER_PARAMETER(float, CameraOverscan)
+		SHADER_PARAMETER_STRUCT_INCLUDE(FSphericalDistortionParams, SphericalDistortionParams)
+		SHADER_PARAMETER_RDG_TEXTURE_UAV(RWTexture2D, OutDistortionMap)
+	END_SHADER_PARAMETER_STRUCT()
+
+	static bool ShouldCompilePermutation(const FGlobalShaderPermutationParameters& Parameters)
+	{
+		return IsFeatureLevelSupported(Parameters.Platform, ERHIFeatureLevel::SM5);
+	}
+};
+
+IMPLEMENT_GLOBAL_SHADER(FSphericalDistortionCS, "/Plugin/CameraCalibrationCore/Private/SphericalDistortion.usf", "MainCS", SF_Compute);
+
 void USphericalLensDistortionModelHandler::InitializeHandler()
 {
 	LensModelClass = USphericalLensModel::StaticClass();
@@ -15,74 +54,72 @@ void USphericalLensDistortionModelHandler::InitializeHandler()
 
 FVector2D USphericalLensDistortionModelHandler::ComputeDistortedUV(const FVector2D& InScreenUV) const
 {
-	// These distances cannot be zero in real-life. If they are, the current distortion state must be bad
-	if ((CurrentState.FocalLengthInfo.FxFy.X == 0.0f) || (CurrentState.FocalLengthInfo.FxFy.Y == 0.0f))
+	const FVector2D FxFy = CurrentState.FocalLengthInfo.FxFy;
+	const FVector2D CxCy = CurrentState.ImageCenter.PrincipalPoint;
+
+	if (FMath::IsNearlyZero(FxFy.X) || FMath::IsNearlyZero(FxFy.Y))
 	{
 		return InScreenUV;
 	}
-	
-	FVector2D NormalizedDistanceFromImageCenter = (InScreenUV - FVector2D(0.5f, 0.5f)) / CurrentState.FocalLengthInfo.FxFy;
 
-	const FVector2D DistanceSquared = NormalizedDistanceFromImageCenter * NormalizedDistanceFromImageCenter;
-	const float RSquared = DistanceSquared.X + DistanceSquared.Y;
+	FVector2D X = (InScreenUV - CxCy) / FxFy; // normalized coords (undistorted)
 
-	// Compute radial distortion
-	const float RadialDistortion = 1.0f + (SphericalParameters.K1 * RSquared) + (SphericalParameters.K2 * RSquared * RSquared) + (SphericalParameters.K3 * RSquared * RSquared * RSquared);
-	
-	// Compute tangential distortion
-	const FVector2D TangentialDistortion = 
-	{
-		(SphericalParameters.P2 * (RSquared + (2.0f * DistanceSquared.X))) + (2.0f * SphericalParameters.P1 * NormalizedDistanceFromImageCenter.X * NormalizedDistanceFromImageCenter.Y),
-		(SphericalParameters.P1 * (RSquared + (2.0f * DistanceSquared.Y))) + (2.0f * SphericalParameters.P2 * NormalizedDistanceFromImageCenter.X * NormalizedDistanceFromImageCenter.Y)
-	};
+	const float R2 = X.X * X.X + X.Y * X.Y;
+	const float R4 = R2 * R2;
+	const float R6 = R4 * R2;
+
+	const float Radial = 1.0f
+		+ SphericalParameters.K1 * R2
+		+ SphericalParameters.K2 * R4
+		+ SphericalParameters.K3 * R6;
+
+	const FVector2D Tangential(
+		SphericalParameters.P2 * (R2 + 2.0f * X.X * X.X) + 2.0f * SphericalParameters.P1 * X.X * X.Y,
+		SphericalParameters.P1 * (R2 + 2.0f * X.Y * X.Y) + 2.0f * SphericalParameters.P2 * X.X * X.Y
+	);
 
-	// Compute total distortion
-	NormalizedDistanceFromImageCenter = NormalizedDistanceFromImageCenter * RadialDistortion + TangentialDistortion;
+	X = X * Radial + Tangential;
 
-	const FVector2D DistortedUV = (NormalizedDistanceFromImageCenter * CurrentState.FocalLengthInfo.FxFy) + CurrentState.ImageCenter.PrincipalPoint;
-	return DistortedUV;
+	return X * FxFy + CxCy; // distorted UV
 }
 
-FVector2D USphericalLensDistortionModelHandler::ComputeUndistortedUV(const FVector2D& InScreenUV) const
+FVector2D USphericalLensDistortionModelHandler::ComputeUndistortedUV(const FVector2D& DistortedUV) const
 {
-	// These distances cannot be zero in real-life. If they are, the current distortion state must be bad
-	if ((CurrentState.FocalLengthInfo.FxFy.X == 0.0f) || (CurrentState.FocalLengthInfo.FxFy.Y == 0.0f))
+	const FVector2D FxFy = CurrentState.FocalLengthInfo.FxFy;
+	const FVector2D CxCy = CurrentState.ImageCenter.PrincipalPoint;
+
+	if (FMath::IsNearlyZero(FxFy.X) || FMath::IsNearlyZero(FxFy.Y))
 	{
-		return InScreenUV;
+		return DistortedUV;
 	}
 
-	FVector2D NormalizedDistanceFromImageCenter = (InScreenUV - CurrentState.ImageCenter.PrincipalPoint) / CurrentState.FocalLengthInfo.FxFy;
-	const FVector2D OriginalDistance = NormalizedDistanceFromImageCenter;
+	const FVector2D Pd = (DistortedUV - CxCy) / FxFy;
 
-	// Iterative approach to distort an undistorted UV using coefficients that were designed to undistort
-	constexpr uint32 NumIterations = 5;
-	for (int32 Index = 0; Index < NumIterations; ++Index)
+	FVector2D P = Pd; // initial guess
+
+	constexpr int32 MaxIters = 5;
+
+	for (int32 It = 0; It < MaxIters; ++It)
 	{
-		const FVector2D DistanceSquared = NormalizedDistanceFromImageCenter * NormalizedDistanceFromImageCenter;
-		const float RSquared = DistanceSquared.X + DistanceSquared.Y;
-
-		const float RadialDistortion = 1.0f + (SphericalParameters.K1 * RSquared) + (SphericalParameters.K2 * RSquared * RSquared) + (SphericalParameters.K3 * RSquared * RSquared * RSquared);
-
-		const FVector2D TangentialDistortion = 
-		{
-			(SphericalParameters.P2 * (RSquared + (2.0f * DistanceSquared.X))) + (2.0f * SphericalParameters.P1 * NormalizedDistanceFromImageCenter.X * NormalizedDistanceFromImageCenter.Y),
-			(SphericalParameters.P1 * (RSquared + (2.0f * DistanceSquared.Y))) + (2.0f * SphericalParameters.P2 * NormalizedDistanceFromImageCenter.X * NormalizedDistanceFromImageCenter.Y)
-		};
-
-		// Guard against divide-by-zero errors
-		if (RadialDistortion == 0.0f)
-		{
-			NormalizedDistanceFromImageCenter = FVector2D(TNumericLimits<float>::Max(), TNumericLimits<float>::Max());
-			break;
-		}
-		else
-		{
-			NormalizedDistanceFromImageCenter = (OriginalDistance - TangentialDistortion) / RadialDistortion;
-		}
-	}
+		const float R2 = P.X * P.X + P.Y * P.Y;
+		const float R4 = R2 * R2;
+		const float R6 = R4 * R2;
+
+		const float Radial = 1.0f
+			+ SphericalParameters.K1 * R2
+			+ SphericalParameters.K2 * R4
+			+ SphericalParameters.K3 * R6;
+
+		const FVector2D Tangential(
+			SphericalParameters.P2 * (R2 + 2.0f * P.X * P.X) + 2.0f * SphericalParameters.P1 * P.X * P.Y,
+			SphericalParameters.P1 * (R2 + 2.0f * P.Y * P.Y) + 2.0f * SphericalParameters.P2 * P.X * P.Y
+		);
 
-	const FVector2D DistortedUV = (NormalizedDistanceFromImageCenter * CurrentState.FocalLengthInfo.FxFy) + FVector2D(0.5f, 0.5f);
-	return DistortedUV;
+		const float InvRadial = FMath::IsNearlyZero(Radial, KINDA_SMALL_NUMBER) ? 1.0f : (1.0f / Radial);
+
+		P = (Pd - Tangential) * InvRadial;
+	}
+	return P * FxFy + CxCy; // undistorted UV
 }
 
 void USphericalLensDistortionModelHandler::InitDistortionMaterials()
@@ -137,3 +174,58 @@ void USphericalLensDistortionModelHandler::InterpretDistortionParameters()
 {
 	LensModelClass->GetDefaultObject<ULensModel>()->FromArray<FSphericalDistortionParameters>(CurrentState.DistortionInfo.Parameters, SphericalParameters);
 }
+
+FString USphericalLensDistortionModelHandler::GetDistortionShaderPath() const
+{
+	return TEXT("/Plugin/CameraCalibrationCore/Private/SphericalDistortion.usf");
+}
+
+
+void USphericalLensDistortionModelHandler::ExecuteDistortionShader(FRDGBuilder& GraphBuilder, const FLensDistortionState& InCurrentState, float InverseOverscan, float CameraOverscan, const FVector2D& SensorSize, FRDGTextureRef& OutDistortionMap) const
+{
+	FIntPoint DistortionMapResolution = OutDistortionMap->Desc.Extent;
+
+	FSphericalDistortionCS::FParameters* ShaderParams = GraphBuilder.AllocParameters<FSphericalDistortionCS::FParameters>();
+	
+	ShaderParams->OutDistortionMap = GraphBuilder.CreateUAV(OutDistortionMap);
+	ShaderParams->ThreadIdToUV = FVector2f(1.0f) / FVector2f(DistortionMapResolution);
+	ShaderParams->InverseOverscan = InverseOverscan;
+	ShaderParams->CameraOverscan = CameraOverscan;
+
+	// Setup spherical distortion parameters
+
+	ShaderParams->SphericalDistortionParams.ImageCenter = FVector2f(InCurrentState.ImageCenter.PrincipalPoint);
+	ShaderParams->SphericalDistortionParams.FocalLength = FVector2f(InCurrentState.FocalLengthInfo.FxFy);
+	
+	// Spherical model requires exactly 5 parameters (K1, K2, K3, P1, P2)
+	if (InCurrentState.DistortionInfo.Parameters.Num() == 5)
+	{
+		ShaderParams->SphericalDistortionParams.K1 = InCurrentState.DistortionInfo.Parameters[0];
+		ShaderParams->SphericalDistortionParams.K2 = InCurrentState.DistortionInfo.Parameters[1];
+		ShaderParams->SphericalDistortionParams.K3 = InCurrentState.DistortionInfo.Parameters[2];
+		ShaderParams->SphericalDistortionParams.P1 = InCurrentState.DistortionInfo.Parameters[3];
+		ShaderParams->SphericalDistortionParams.P2 = InCurrentState.DistortionInfo.Parameters[4];
+	}
+	else
+	{
+		UE_LOG(LogCameraCalibrationCore, Warning, 
+			TEXT("SphericalLensDistortionModelHandler: Invalid parameter count (%d), expected exactly 5. Using default values."), 
+			InCurrentState.DistortionInfo.Parameters.Num()
+		);
+
+		ShaderParams->SphericalDistortionParams.K1 = 0.0f;
+		ShaderParams->SphericalDistortionParams.K2 = 0.0f;
+		ShaderParams->SphericalDistortionParams.K3 = 0.0f;
+		ShaderParams->SphericalDistortionParams.P1 = 0.0f;
+		ShaderParams->SphericalDistortionParams.P2 = 0.0f;
+	}
+
+	TShaderMapRef<FSphericalDistortionCS> ComputeShader(GetGlobalShaderMap(GMaxRHIFeatureLevel));
+	FComputeShaderUtils::AddPass(
+		GraphBuilder,
+		RDG_EVENT_NAME("SphericalDistortionDisplacementMap"),
+		ComputeShader,
+		ShaderParams,
+		FIntVector(FMath::DivideAndRoundUp(DistortionMapResolution.X, 8), FMath::DivideAndRoundUp(DistortionMapResolution.Y, 8), 1));
+}
+
diff --git a/Engine/Plugins/VirtualProduction/CameraCalibrationCore/Source/CameraCalibrationCore/Public/AnamorphicLensDistortionModelHandler.h b/Engine/Plugins/VirtualProduction/CameraCalibrationCore/Source/CameraCalibrationCore/Public/AnamorphicLensDistortionModelHandler.h
--- a/Engine/Plugins/VirtualProduction/CameraCalibrationCore/Source/CameraCalibrationCore/Public/AnamorphicLensDistortionModelHandler.h
+++ b/Engine/Plugins/VirtualProduction/CameraCalibrationCore/Source/CameraCalibrationCore/Public/AnamorphicLensDistortionModelHandler.h
@@ -24,6 +24,9 @@ protected:
 	UE_API virtual void InitDistortionMaterials() override;
 	UE_API virtual void UpdateMaterialParameters() override;
 	UE_API virtual void InterpretDistortionParameters() override;
+	UE_API virtual FString GetDistortionShaderPath() const override;
+	UE_API virtual void ExecuteDistortionShader(class FRDGBuilder& GraphBuilder, const FLensDistortionState& InCurrentState, float InverseOverscan, float CameraOverscan, const FVector2D& SensorSize, FRDGTextureRef& OutDistortionMap) const override;
+	UE_API virtual bool IsForwardDistorting() const override { return false; }
 	//~ End ULensDistortionModelHandlerBase interface
 
 private:
diff --git a/Engine/Plugins/VirtualProduction/CameraCalibrationCore/Source/CameraCalibrationCore/Public/LensDistortionModelHandlerBase.h b/Engine/Plugins/VirtualProduction/CameraCalibrationCore/Source/CameraCalibrationCore/Public/LensDistortionModelHandlerBase.h
--- a/Engine/Plugins/VirtualProduction/CameraCalibrationCore/Source/CameraCalibrationCore/Public/LensDistortionModelHandlerBase.h
+++ b/Engine/Plugins/VirtualProduction/CameraCalibrationCore/Source/CameraCalibrationCore/Public/LensDistortionModelHandlerBase.h
@@ -11,6 +11,11 @@
 
 #include "LensDistortionModelHandlerBase.generated.h"
 
+// Forward declarations for RDG types
+class FRDGBuilder;
+class FRDGTexture;
+using FRDGTextureRef = FRDGTexture*;
+
 #define UE_API CAMERACALIBRATIONCORE_API
 
 USTRUCT(BlueprintType)
@@ -127,6 +132,12 @@ public:
 	/** Create two displacement maps, for distortion and undistortion, using the input resolution */
 	UE_API void CreateDisplacementMaps(const FIntPoint DisplacementMapResolution);
 
+	/** Execute the distortion shader with the specified parameters, which enables a modular SVE that doesn't need to know lens model specifics. */
+	virtual void ExecuteDistortionShader(class FRDGBuilder& GraphBuilder, const FLensDistortionState& InCurrentState, float InverseOverscan, float CameraOverscan, const FVector2D& SensorSize, FRDGTextureRef& OutDistortionMap) const PURE_VIRTUAL(ULensDistortionModelHandlerBase::ExecuteDistortionShader);
+
+	/** Returns true if this model performs forward distortion (needs inverse overscan), false if reverse distortion (needs forward overscan) */
+	virtual bool IsForwardDistorting() const { return true; }
+
 protected:
 	/** Initialize the handler. Derived classes must set the LensModelClass that they support, if not already set */
 	virtual void InitializeHandler() PURE_VIRTUAL(ULensDistortionModelHandlerBase::InitializeHandler);
@@ -146,6 +157,9 @@ protected:
 	/** Convert the generic distortion parameter array into the specific structure of parameters used by the supported lens model */
 	virtual void InterpretDistortionParameters() PURE_VIRTUAL(ULensDistortionModelHandlerBase::InterpretDistortionParameters);
 
+	/** Return the shader path for this model's distortion compute shader, promoting a modular shader architecture */
+	virtual FString GetDistortionShaderPath() const PURE_VIRTUAL(ULensDistortionModelHandlerBase::GetDistortionShaderPath, return FString(););
+
 private:
 	/** Tests whether the input material's shaders have been compiled, indicating the material can be used for drawing */
 	UE_API bool IsDisplacementMapMaterialReady(UMaterialInstanceDynamic* MID);
diff --git a/Engine/Plugins/VirtualProduction/CameraCalibrationCore/Source/CameraCalibrationCore/Public/LensDistortionSceneViewExtension.h b/Engine/Plugins/VirtualProduction/CameraCalibrationCore/Source/CameraCalibrationCore/Public/LensDistortionSceneViewExtension.h
--- a/Engine/Plugins/VirtualProduction/CameraCalibrationCore/Source/CameraCalibrationCore/Public/LensDistortionSceneViewExtension.h
+++ b/Engine/Plugins/VirtualProduction/CameraCalibrationCore/Source/CameraCalibrationCore/Public/LensDistortionSceneViewExtension.h
@@ -14,23 +14,6 @@ class UCameraComponent;
  */
 class FLensDistortionSceneViewExtension : public FSceneViewExtensionBase
 {
-public:
-	/** The distortion model to use when creating distortion maps */
-	enum class EDistortionModel : uint8
-	{
-		None,
-		
-		/** Distortion maps generated from distortion parameters with a spherical distortion equation */
-		SphericalDistortion,
-
-		/** Distortion maps generated from distortion parameters with an anamorphic distortion equation */
-		AnamorphicDistortion,
-
-		/** Distortion maps generated from UV-based ST maps */
-		STMap,
-		MAX
-	};
-	
 public:
 
 	FLensDistortionSceneViewExtension(const FAutoRegister& AutoRegister);
@@ -54,10 +37,10 @@ public:
 
 private:
 	/** Use the input distortion state to draw a distortion displacement map */
-	void DrawDisplacementMap_RenderThread(FRDGBuilder& GraphBuilder, const FLensDistortionState& CurrentState, EDistortionModel DistortionModel, float InverseOverscan, float CameraOverscan, const FVector2D& SensorSize, FRDGTextureRef& OutDistortionMapWithOverscan) const;
+	void DrawDisplacementMap_RenderThread(FRDGBuilder& GraphBuilder, const FLensDistortionState& CurrentState, class ULensDistortionModelHandlerBase* ModelHandler, float InverseOverscan, float CameraOverscan, const FVector2D& SensorSize, FRDGTextureRef& OutDistortionMapWithOverscan) const;
 
 	/** Use the input blend parameters to draw multiple displacement maps and blend them together into a final distortion displacement map */
-	void BlendDisplacementMaps_RenderThread(FRDGBuilder& GraphBuilder, const FDisplacementMapBlendingParams& BlendState, EDistortionModel DistortionModel, float InverseOverscan, float CameraOverscan, const FVector2D& SensorSize, FRDGTextureRef& OutDistortionMapWithOverscan) const;
+	void BlendDisplacementMaps_RenderThread(FRDGBuilder& GraphBuilder, const FDisplacementMapBlendingParams& BlendState, class ULensDistortionModelHandlerBase* ModelHandler, float InverseOverscan, float CameraOverscan, const FVector2D& SensorSize, FRDGTextureRef& OutDistortionMapWithOverscan) const;
 
 	/** Crop the input overscanned distortion map to the original requested resolution */
 	void CropDisplacementMap_RenderThread(FRDGBuilder& GraphBuilder, const FRDGTextureRef& InDistortionMapWithOverscan, FRDGTextureRef& OutDistortionMap) const;
@@ -71,11 +54,6 @@ private:
 	/** Recenters the specified distortion map within an overscanned map, filling the overscanned region with zeros */
 	void RecenterSTDisplacementMap_RenderThread(FRDGBuilder& GraphBuilder, const FRDGTextureRef& InDisplacementMap, float InOverscan, FRDGTextureRef& OutRecenteredDisplacementMap) const;
 	
-	/** Gets whether the specified distortion model is forward distorting (generating the distort map first, then inverting for the undistort) or inverse distorting */
-	bool IsDistortionModelForwardDistorting(EDistortionModel InDistortionModel) const;
-
-	/** Computes the overscan needed to invert the distortion map */
-	float GetInverseOverscan(ULensDistortionModelHandlerBase* LensDistortionHandler, EDistortionModel InDistortionModel) const;
 	
 private:
 
diff --git a/Engine/Plugins/VirtualProduction/CameraCalibrationCore/Source/CameraCalibrationCore/Public/SphericalLensDistortionModelHandler.h b/Engine/Plugins/VirtualProduction/CameraCalibrationCore/Source/CameraCalibrationCore/Public/SphericalLensDistortionModelHandler.h
--- a/Engine/Plugins/VirtualProduction/CameraCalibrationCore/Source/CameraCalibrationCore/Public/SphericalLensDistortionModelHandler.h
+++ b/Engine/Plugins/VirtualProduction/CameraCalibrationCore/Source/CameraCalibrationCore/Public/SphericalLensDistortionModelHandler.h
@@ -24,6 +24,8 @@ protected:
 	UE_API virtual void InitDistortionMaterials() override;
 	UE_API virtual void UpdateMaterialParameters() override;
 	UE_API virtual void InterpretDistortionParameters() override;
+	UE_API virtual FString GetDistortionShaderPath() const override;
+	UE_API virtual void ExecuteDistortionShader(class FRDGBuilder& GraphBuilder, const FLensDistortionState& InCurrentState, float InverseOverscan, float CameraOverscan, const FVector2D& SensorSize, FRDGTextureRef& OutDistortionMap) const override;
 	//~ End ULensDistortionModelHandlerBase interface
 
 private:
