CameraCalibration: Fix possible unsafe ModelHandler dereferencing outside of the lock.

diff --git a/Engine/Plugins/VirtualProduction/CameraCalibrationCore/Source/CameraCalibrationCore/Private/LensDistortionSceneViewExtension.cpp b/Engine/Plugins/VirtualProduction/CameraCalibrationCore/Source/CameraCalibrationCore/Private/LensDistortionSceneViewExtension.cpp
--- a/Engine/Plugins/VirtualProduction/CameraCalibrationCore/Source/CameraCalibrationCore/Private/LensDistortionSceneViewExtension.cpp
+++ b/Engine/Plugins/VirtualProduction/CameraCalibrationCore/Source/CameraCalibrationCore/Private/LensDistortionSceneViewExtension.cpp
@@ -551,12 +551,23 @@ bool FLensDistortionSceneViewExtension::RenderViewDistortionLUT(FRDGBuilder& Gra
 	TStrongObjectPtr<ULensDistortionModelHandlerBase> ModelHandler = Proxy->LensDistortionHandler.Pin();
 	
 	bool bIsSTMap = false;
+	bool bForwardDistort = true;
+	float RawInverseOverscan = 1.0f;
+	FTextureRHIRef UndistortionMapTextureRef = nullptr;
+	FTextureRHIRef DistortionMapTextureRef = nullptr;
+	
 	if (ModelHandler.IsValid())
 	{
 		if (!Proxy->Params.States[0].DistortionInfo.Parameters.Num())
 		{
 			bIsSTMap = true;
+
+			UndistortionMapTextureRef = ModelHandler->GetDistortionDisplacementMap()->GetResource()->GetTexture2DRHI();
+			DistortionMapTextureRef = ModelHandler->GetUndistortionDisplacementMap()->GetResource()->GetTexture2DRHI();
 		}
+		
+		bForwardDistort = ModelHandler->IsForwardDistorting();
+		RawInverseOverscan = bForwardDistort ? ModelHandler->ComputeOverscanFactor() : ModelHandler->ComputeInverseOverscanFactor();
 	}
 	
 	if (bIsSTMap)
@@ -566,10 +577,8 @@ bool FLensDistortionSceneViewExtension::RenderViewDistortionLUT(FRDGBuilder& Gra
 		// Further, ST maps may have intrinsic inaccuracies that can cause artifacting in TSR, so use a warp grid inversion to generate a fully approximate
 		// distortion ST map from the undistortion map. We also need to recenter the undistortion map to correspond to the overscan, as well as scaling its displacement
 		// by the overscan amount
-		FTextureRHIRef UndistortionMapTextureRef = ModelHandler->GetDistortionDisplacementMap()->GetResource()->GetTexture2DRHI();
 		FRDGTextureRef OriginalUndistortionMap = GraphBuilder.RegisterExternalTexture(CreateRenderTarget(UndistortionMapTextureRef, TEXT("OriginalUndistortionDisplacementMap")));
 
-		FTextureRHIRef DistortionMapTextureRef = ModelHandler->GetUndistortionDisplacementMap()->GetResource()->GetTexture2DRHI();
 		FRDGTextureRef OriginalDistortionMap = GraphBuilder.RegisterExternalTexture(CreateRenderTarget(DistortionMapTextureRef, TEXT("OriginalDistortionDisplacementMap")));
 
 		FRDGTextureDesc FilledDistortionMapDesc = FRDGTextureDesc::Create2D(
@@ -622,16 +631,12 @@ bool FLensDistortionSceneViewExtension::RenderViewDistortionLUT(FRDGBuilder& Gra
 		FClearValueBinding::Black,
 		TexCreate_ShaderResource | TexCreate_RenderTargetable);
 		
-	// Choose texture descriptors based on distortion direction (matching original logic)
-	const bool bForwardDistort = ModelHandler->IsForwardDistorting();
+	// Choose texture descriptors based on distortion direction
 	FRDGTextureDesc DistortionMapDesc = bForwardDistort ? ForwardDistortionMapDesc : InverseDistortionMapDesc;
 	FRDGTextureDesc UndistortionMapDesc = bForwardDistort ? InverseDistortionMapDesc : ForwardDistortionMapDesc;
 	
 	ViewDistortionLUT.DistortingDisplacementTexture = GraphBuilder.CreateTexture(DistortionMapDesc, TEXT("DistortionDisplacementMap"));
 	ViewDistortionLUT.UndistortingDisplacementTexture = GraphBuilder.CreateTexture(UndistortionMapDesc, TEXT("UndistortionDisplacementMap"));
-
-	// Let each handler compute its own overscan - handlers know their own mathematical properties
-	float RawInverseOverscan = ModelHandler->IsForwardDistorting() ? ModelHandler->ComputeOverscanFactor() : ModelHandler->ComputeInverseOverscanFactor();
 	float InverseOverscan = FMath::Clamp(RawInverseOverscan, 0.0f, 2.0f);
 	
 	// Adjust the overscan resolution to be square, with each side being a multiple of 8
