CameraCalibration: Relaxed new lens handler api to have default implementations of the new functions.

This avoids compilation errors on 3rd party lens models that don't implement SVE. They couldn't do it until recently the shader can be executed by the handler itself.

#rnx

diff --git a/Engine/Plugins/VirtualProduction/CameraCalibrationCore/Source/CameraCalibrationCore/Public/LensDistortionModelHandlerBase.h b/Engine/Plugins/VirtualProduction/CameraCalibrationCore/Source/CameraCalibrationCore/Public/LensDistortionModelHandlerBase.h
--- a/Engine/Plugins/VirtualProduction/CameraCalibrationCore/Source/CameraCalibrationCore/Public/LensDistortionModelHandlerBase.h
+++ b/Engine/Plugins/VirtualProduction/CameraCalibrationCore/Source/CameraCalibrationCore/Public/LensDistortionModelHandlerBase.h
@@ -133,7 +133,7 @@ public:
 	UE_API void CreateDisplacementMaps(const FIntPoint DisplacementMapResolution);
 
 	/** Execute the distortion shader with the specified parameters, which enables a modular SVE that doesn't need to know lens model specifics. */
-	virtual void ExecuteDistortionShader(class FRDGBuilder& GraphBuilder, const FLensDistortionState& InCurrentState, float InverseOverscan, float CameraOverscan, const FVector2D& SensorSize, FRDGTextureRef& OutDistortionMap) const PURE_VIRTUAL(ULensDistortionModelHandlerBase::ExecuteDistortionShader);
+	virtual void ExecuteDistortionShader(class FRDGBuilder& GraphBuilder, const FLensDistortionState& InCurrentState, float InverseOverscan, float CameraOverscan, const FVector2D& SensorSize, FRDGTextureRef& OutDistortionMap) const {}
 
 	/** Returns true if this model performs forward distortion (needs inverse overscan), false if reverse distortion (needs forward overscan) */
 	virtual bool IsForwardDistorting() const { return true; }
@@ -158,7 +158,7 @@ protected:
 	virtual void InterpretDistortionParameters() PURE_VIRTUAL(ULensDistortionModelHandlerBase::InterpretDistortionParameters);
 
 	/** Return the shader path for this model's distortion compute shader, promoting a modular shader architecture */
-	virtual FString GetDistortionShaderPath() const PURE_VIRTUAL(ULensDistortionModelHandlerBase::GetDistortionShaderPath, return FString(););
+	virtual FString GetDistortionShaderPath() const { return FString(); }
 
 private:
 	/** Tests whether the input material's shaders have been compiled, indicating the material can be used for drawing */
