CameraCalibration: Fix possible array out of bounds.

#rnx

diff --git a/Engine/Plugins/VirtualProduction/CameraCalibrationCore/Source/CameraCalibrationCore/Private/DistortionRenderingUtils.cpp b/Engine/Plugins/VirtualProduction/CameraCalibrationCore/Source/CameraCalibrationCore/Private/DistortionRenderingUtils.cpp
--- a/Engine/Plugins/VirtualProduction/CameraCalibrationCore/Source/CameraCalibrationCore/Private/DistortionRenderingUtils.cpp
+++ b/Engine/Plugins/VirtualProduction/CameraCalibrationCore/Source/CameraCalibrationCore/Private/DistortionRenderingUtils.cpp
@@ -70,6 +70,12 @@ namespace DistortionRenderingUtils
 
 		const int32 NumPoints = ImagePoints.Num();
 		const FTextureRenderTargetResource* const DistortionMapResource = DistortionMap->GameThread_GetRenderTargetResource();
+		
+		// Ensure the output array has enough space for indexed writes
+		if (OutUndistortedPoints.Num() != NumPoints)
+		{
+			OutUndistortedPoints.SetNumUninitialized(NumPoints);
+		}
 
 		ENQUEUE_RENDER_COMMAND(DistortionRenderingUtils_UndistortImagePoints)(
 			[DistortionMapResource, NumPoints, Points = MoveTemp(ImagePoints), &OutUndistortedPoints](FRHICommandListImmediate& RHICmdList)
