CameraCalibration: Handle returning uninitialized array when gpu readback fails.

#rnx

diff --git a/Engine/Plugins/VirtualProduction/CameraCalibrationCore/Source/CameraCalibrationCore/Private/DistortionRenderingUtils.cpp b/Engine/Plugins/VirtualProduction/CameraCalibrationCore/Source/CameraCalibrationCore/Private/DistortionRenderingUtils.cpp
--- a/Engine/Plugins/VirtualProduction/CameraCalibrationCore/Source/CameraCalibrationCore/Private/DistortionRenderingUtils.cpp
+++ b/Engine/Plugins/VirtualProduction/CameraCalibrationCore/Source/CameraCalibrationCore/Private/DistortionRenderingUtils.cpp
@@ -3,6 +3,7 @@
 
 #include "DistortionRenderingUtils.h"
 
+#include "CameraCalibrationCoreLog.h"
 #include "Containers/DynamicRHIResourceArray.h"
 #include "Engine/TextureRenderTarget2D.h"
 #include "GlobalShader.h"
@@ -71,11 +72,8 @@ namespace DistortionRenderingUtils
 		const int32 NumPoints = ImagePoints.Num();
 		const FTextureRenderTargetResource* const DistortionMapResource = DistortionMap->GameThread_GetRenderTargetResource();
 		
-		// Ensure the output array has enough space for indexed writes
-		if (OutUndistortedPoints.Num() != NumPoints)
-		{
-			OutUndistortedPoints.SetNumUninitialized(NumPoints);
-		}
+		// Ensure the output array has the correct size
+		OutUndistortedPoints.SetNum(NumPoints);
 
 		ENQUEUE_RENDER_COMMAND(DistortionRenderingUtils_UndistortImagePoints)(
 			[DistortionMapResource, NumPoints, Points = MoveTemp(ImagePoints), &OutUndistortedPoints](FRHICommandListImmediate& RHICmdList)
@@ -129,6 +127,15 @@ namespace DistortionRenderingUtils
 						OutUndistortedPoints[Index] = UndistortedPointData[Index];
 					}
 				}
+				else
+				{
+					UE_LOG(LogCameraCalibrationCore, Warning, TEXT("DistortionRenderingUtils::UndistortImagePoints: Failed to read undistorted points from GPU, so filling with zeroes."));
+
+					for (int32 Index = 0; Index < NumPoints; ++Index)
+					{
+						OutUndistortedPoints[Index] = FVector2f::ZeroVector;
+					}
+				}
 
 				DestinationStagingBuffer->Unlock();
 			});
