CameraCalibration: Add Brown Conrady U-D model using the full 6 parameters.

Next will be the D-U model, which would complete what's needed to support the current OpenLensIO models used by OpenTrackIO.

[RN] Added polynomial division brown-conrady model in the U-D direction.

#virtualized

diff --git a/Engine/Plugins/VirtualProduction/CameraCalibrationCore/Shaders/Private/BrownConradyUDDistortion.usf b/Engine/Plugins/VirtualProduction/CameraCalibrationCore/Shaders/Private/BrownConradyUDDistortion.usf
new file mode 100644
--- /dev/null
+++ b/Engine/Plugins/VirtualProduction/CameraCalibrationCore/Shaders/Private/BrownConradyUDDistortion.usf
@@ -0,0 +1,106 @@
+// Copyright Epic Games, Inc. All Rights Reserved.
+
+/*=============================================================================
+	Implements the Brown-Conrady U-D model with 6 radial and 2 tangential components.
+	Uses polynomial division model: (1 + K1*r^2 + K2*r^4 + K3*r^6) / (1 + K4*r^2 + K5*r^4 + K6*r^6)
+=============================================================================*/
+
+#include "/Engine/Public/Platform.ush"
+#include "/Engine/Private/Common.ush"
+
+// Shared parameters
+float2 ThreadIdToUV;
+float InverseOverscan;
+float CameraOverscan;
+RWTexture2D<float2> OutDistortionMap;
+
+// Brown-Conrady U-D distortion parameters (polynomial division model)
+float2 FocalLength;
+float2 ImageCenter;
+float K1;
+float K2;
+float K3;
+float K4;
+float K5;
+float K6;
+float P1;
+float P2;
+
+float2 BrownConradyUDDistort(float2 InputUV)
+{
+	float x_prime = (InputUV.x - ImageCenter.x) / FocalLength.x;
+	float y_prime = (InputUV.y - ImageCenter.y) / FocalLength.y;
+
+	float x_prime2 = x_prime * x_prime;
+	float y_prime2 = y_prime * y_prime;
+
+	float r2 = x_prime2 + y_prime2;
+	float r4 = r2 * r2;
+	float r6 = r2 * r2 * r2;
+
+	// Compute radial distortion using polynomial division model
+	float dr_numerator = 1 + K1 * r2 + K2 * r4 + K3 * r6;
+	float dr_denominator = 1 + K4 * r2 + K5 * r4 + K6 * r6;
+	float dr = (abs(dr_denominator) < 0.00001) ? dr_numerator : (dr_numerator / dr_denominator);
+
+	// Compute tangential distortion
+	float dtx = 2 * P1 * x_prime * y_prime + P2 * (r2 + 2 * x_prime2);
+	float dty = P1 * (r2 + 2 * y_prime2) + 2 * P2 * x_prime * y_prime;
+
+	// Compute total distortion
+	float x_double_prime = (x_prime * dr) + dtx;
+	float y_double_prime = (y_prime * dr) + dty;
+
+	float x_undistort = (FocalLength.x * x_double_prime) + ImageCenter.x;
+	float y_undistort = (FocalLength.y * y_double_prime) + ImageCenter.y;
+
+	float2 result;
+	result.x = x_undistort;
+	result.y = y_undistort;
+
+	return result;
+}
+
+float2 Distort(float2 InputUV)
+{
+	return BrownConradyUDDistort(InputUV);
+}
+
+bool IsDistortionWellOriented(float2 InputUV, float2 DistortedUV, float2 UVOffset)
+{
+	float2 DistortedXOffset = Distort(InputUV + UVOffset * int2(1, 0));
+	float2 DistortedYOffset = Distort(InputUV + UVOffset * int2(0, 1));
+
+	float2 XOffset = DistortedXOffset - DistortedUV;
+	float2 YOffset = DistortedYOffset - DistortedUV;
+	
+	bool bIsFrontFacing = (XOffset.x * YOffset.y - XOffset.y * YOffset.x) >= 0;
+	bool bIsRightSideUp = (XOffset.x * UVOffset.x) >= 0 && (YOffset.y * UVOffset.y) >= 0;
+	return bIsFrontFacing && bIsRightSideUp;
+}
+
+float2 TryDistort(float2 InputUV)
+{
+	float2 DistortedPosition = Distort(InputUV);
+	float2 UVOffset = 0.5 * ThreadIdToUV;
+	bool bIsOrientedWell =
+		IsDistortionWellOriented(InputUV, DistortedPosition, UVOffset) &&
+		IsDistortionWellOriented(InputUV, DistortedPosition, -UVOffset);
+
+	return lerp(InputUV, DistortedPosition, bIsOrientedWell);
+}
+
+[numthreads(8, 8, 1)]
+void MainCS(uint3 DispatchThreadId : SV_DispatchThreadID)
+{
+	// Add half-pixel offset to ThreadID to get accurate UV position
+	float2 UndistortedUV = (DispatchThreadId.xy + 0.5) * ThreadIdToUV;
+	
+	// Rescale the UV by an overscan factor (>= 1.0) to widen the range of UVs to be distorted
+	UndistortedUV = ((UndistortedUV - 0.5) * InverseOverscan * CameraOverscan) + 0.5;
+	
+	float2 DistortedUV = TryDistort(UndistortedUV);
+	
+	// Compute the UV displacement and scale it by the CameraOverscan to account for the increased resolution of the overscanned scene texture 
+	OutDistortionMap[DispatchThreadId.xy] = (DistortedUV - UndistortedUV) / CameraOverscan;
+}
\ No newline at end of file
diff --git a/Engine/Plugins/VirtualProduction/CameraCalibrationCore/Source/CameraCalibrationCore/Private/BrownConradyUDLensDistortionModelHandler.cpp b/Engine/Plugins/VirtualProduction/CameraCalibrationCore/Source/CameraCalibrationCore/Private/BrownConradyUDLensDistortionModelHandler.cpp
new file mode 100644
--- /dev/null
+++ b/Engine/Plugins/VirtualProduction/CameraCalibrationCore/Source/CameraCalibrationCore/Private/BrownConradyUDLensDistortionModelHandler.cpp
@@ -0,0 +1,256 @@
+// Copyright Epic Games, Inc. All Rights Reserved.
+
+#include "BrownConradyUDLensDistortionModelHandler.h"
+
+#include "CameraCalibrationCoreLog.h"
+#include "CameraCalibrationSettings.h"
+#include "Engine/TextureRenderTarget2D.h"
+#include "GlobalShader.h"
+#include "Math/NumericLimits.h"
+#include "RenderGraphBuilder.h"
+#include "RenderGraphUtils.h"
+#include "ShaderParameterUtils.h"
+#include "ShaderCompilerCore.h"
+
+#include UE_INLINE_GENERATED_CPP_BY_NAME(BrownConradyUDLensDistortionModelHandler)
+
+// Brown-Conrady U-D distortion shader parameter structure
+BEGIN_SHADER_PARAMETER_STRUCT(FBrownConradyUDDistortionParams, )
+	SHADER_PARAMETER(FVector2f, FocalLength)
+	SHADER_PARAMETER(FVector2f, ImageCenter)
+	SHADER_PARAMETER(float, K1)
+	SHADER_PARAMETER(float, K2)
+	SHADER_PARAMETER(float, K3)
+	SHADER_PARAMETER(float, K4)
+	SHADER_PARAMETER(float, K5)
+	SHADER_PARAMETER(float, K6)
+	SHADER_PARAMETER(float, P1)
+	SHADER_PARAMETER(float, P2)
+END_SHADER_PARAMETER_STRUCT()
+
+// Dedicated Brown-Conrady U-D distortion compute shader
+class FBrownConradyUDDistortionCS : public FGlobalShader
+{
+	DECLARE_GLOBAL_SHADER(FBrownConradyUDDistortionCS);
+	SHADER_USE_PARAMETER_STRUCT(FBrownConradyUDDistortionCS, FGlobalShader);
+
+	BEGIN_SHADER_PARAMETER_STRUCT(FParameters, )
+		SHADER_PARAMETER(FVector2f, ThreadIdToUV)
+		SHADER_PARAMETER(float, InverseOverscan)
+		SHADER_PARAMETER(float, CameraOverscan)
+		SHADER_PARAMETER_STRUCT_INCLUDE(FBrownConradyUDDistortionParams, BrownConradyUDDistortionParams)
+		SHADER_PARAMETER_RDG_TEXTURE_UAV(RWTexture2D, OutDistortionMap)
+	END_SHADER_PARAMETER_STRUCT()
+
+	static bool ShouldCompilePermutation(const FGlobalShaderPermutationParameters& Parameters)
+	{
+		return GetMaxSupportedFeatureLevel(Parameters.Platform) >= ERHIFeatureLevel::SM5;
+	}
+};
+
+IMPLEMENT_GLOBAL_SHADER(FBrownConradyUDDistortionCS, "/Plugin/CameraCalibrationCore/Private/BrownConradyUDDistortion.usf", "MainCS", SF_Compute);
+
+void UBrownConradyUDLensDistortionModelHandler::InitializeHandler()
+{
+	LensModelClass = UBrownConradyUDLensModel::StaticClass();
+}
+
+FVector2D UBrownConradyUDLensDistortionModelHandler::ComputeDistortedUV(const FVector2D& InScreenUV) const
+{
+	const FVector2D FxFy = CurrentState.FocalLengthInfo.FxFy;
+	const FVector2D CxCy = CurrentState.ImageCenter.PrincipalPoint;
+
+	if (FMath::IsNearlyZero(FxFy.X) || FMath::IsNearlyZero(FxFy.Y))
+	{
+		return InScreenUV;
+	}
+
+	FVector2D X = (InScreenUV - CxCy) / FxFy; // normalized coords (undistorted)
+
+	const float R2 = X.X * X.X + X.Y * X.Y;
+	const float R4 = R2 * R2;
+	const float R6 = R4 * R2;
+
+	const float RadialNumerator = 1.0f
+		+ BrownConradyUDParameters.K1 * R2
+		+ BrownConradyUDParameters.K2 * R4
+		+ BrownConradyUDParameters.K3 * R6;
+	
+	const float RadialDenominator = 1.0f
+		+ BrownConradyUDParameters.K4 * R2
+		+ BrownConradyUDParameters.K5 * R4
+		+ BrownConradyUDParameters.K6 * R6;
+	
+	const float Radial = FMath::IsNearlyZero(RadialDenominator, KINDA_SMALL_NUMBER) ? RadialNumerator : (RadialNumerator / RadialDenominator);
+
+	const FVector2D Tangential(
+		BrownConradyUDParameters.P2 * (R2 + 2.0f * X.X * X.X) + 2.0f * BrownConradyUDParameters.P1 * X.X * X.Y,
+		BrownConradyUDParameters.P1 * (R2 + 2.0f * X.Y * X.Y) + 2.0f * BrownConradyUDParameters.P2 * X.X * X.Y
+	);
+
+	X = X * Radial + Tangential;
+
+	return X * FxFy + CxCy; // distorted UV
+}
+
+FVector2D UBrownConradyUDLensDistortionModelHandler::ComputeUndistortedUV(const FVector2D& DistortedUV) const
+{
+	const FVector2D FxFy = CurrentState.FocalLengthInfo.FxFy;
+	const FVector2D CxCy = CurrentState.ImageCenter.PrincipalPoint;
+
+	if (FMath::IsNearlyZero(FxFy.X) || FMath::IsNearlyZero(FxFy.Y))
+	{
+		return DistortedUV;
+	}
+
+	const FVector2D Pd = (DistortedUV - CxCy) / FxFy;
+
+	FVector2D P = Pd; // initial guess
+
+	constexpr int32 MaxIters = 10;
+
+	for (int32 It = 0; It < MaxIters; ++It)
+	{
+		const float R2 = P.X * P.X + P.Y * P.Y;
+		const float R4 = R2 * R2;
+		const float R6 = R4 * R2;
+
+		const float RadialNumerator = 1.0f
+			+ BrownConradyUDParameters.K1 * R2
+			+ BrownConradyUDParameters.K2 * R4
+			+ BrownConradyUDParameters.K3 * R6;
+		
+		const float RadialDenominator = 1.0f
+			+ BrownConradyUDParameters.K4 * R2
+			+ BrownConradyUDParameters.K5 * R4
+			+ BrownConradyUDParameters.K6 * R6;
+		
+		const float Radial = FMath::IsNearlyZero(RadialDenominator, KINDA_SMALL_NUMBER) ? RadialNumerator : (RadialNumerator / RadialDenominator);
+
+		const FVector2D Tangential(
+			BrownConradyUDParameters.P2 * (R2 + 2.0f * P.X * P.X) + 2.0f * BrownConradyUDParameters.P1 * P.X * P.Y,
+			BrownConradyUDParameters.P1 * (R2 + 2.0f * P.Y * P.Y) + 2.0f * BrownConradyUDParameters.P2 * P.X * P.Y
+		);
+
+		const float InvRadial = FMath::IsNearlyZero(Radial, KINDA_SMALL_NUMBER) ? 1.0f : (1.0f / Radial);
+
+		P = (Pd - Tangential) * InvRadial;
+	}
+	return P * FxFy + CxCy; // undistorted UV
+}
+
+void UBrownConradyUDLensDistortionModelHandler::InitDistortionMaterials()
+{
+	if (DistortionPostProcessMID == nullptr)
+	{
+		UMaterialInterface* DistortionMaterialParent = GetDefault<UCameraCalibrationSettings>()->GetDefaultDistortionMaterial(this->StaticClass());
+		DistortionPostProcessMID = UMaterialInstanceDynamic::Create(DistortionMaterialParent, this);
+	}
+
+	if (UndistortionDisplacementMapMID == nullptr)
+	{
+		UMaterialInterface* MaterialParent = GetDefault<UCameraCalibrationSettings>()->GetDefaultUndistortionDisplacementMaterial(this->StaticClass());
+		UndistortionDisplacementMapMID = UMaterialInstanceDynamic::Create(MaterialParent, this);
+	}
+
+	if (DistortionDisplacementMapMID == nullptr)
+	{
+		UMaterialInterface* MaterialParent = GetDefault<UCameraCalibrationSettings>()->GetDefaultDistortionDisplacementMaterial(this->StaticClass());
+		DistortionDisplacementMapMID = UMaterialInstanceDynamic::Create(MaterialParent, this);
+	}
+
+	DistortionPostProcessMID->SetTextureParameterValue("UndistortionDisplacementMap", UndistortionDisplacementMapRT);
+	DistortionPostProcessMID->SetTextureParameterValue("DistortionDisplacementMap", DistortionDisplacementMapRT);
+
+	SetDistortionState(CurrentState);
+}
+
+void UBrownConradyUDLensDistortionModelHandler::UpdateMaterialParameters()
+{
+	//Helper function to set material parameters of an MID
+	const auto SetDistortionMaterialParameters = [this](UMaterialInstanceDynamic* const MID)
+	{
+		MID->SetScalarParameterValue("k1", BrownConradyUDParameters.K1);
+		MID->SetScalarParameterValue("k2", BrownConradyUDParameters.K2);
+		MID->SetScalarParameterValue("k3", BrownConradyUDParameters.K3);
+		MID->SetScalarParameterValue("k4", BrownConradyUDParameters.K4);
+		MID->SetScalarParameterValue("k5", BrownConradyUDParameters.K5);
+		MID->SetScalarParameterValue("k6", BrownConradyUDParameters.K6);
+		MID->SetScalarParameterValue("p1", BrownConradyUDParameters.P1);
+		MID->SetScalarParameterValue("p2", BrownConradyUDParameters.P2);
+
+		MID->SetScalarParameterValue("cx", CurrentState.ImageCenter.PrincipalPoint.X);
+		MID->SetScalarParameterValue("cy", CurrentState.ImageCenter.PrincipalPoint.Y);
+
+		MID->SetScalarParameterValue("fx", CurrentState.FocalLengthInfo.FxFy.X);
+		MID->SetScalarParameterValue("fy", CurrentState.FocalLengthInfo.FxFy.Y);
+	};
+
+	SetDistortionMaterialParameters(UndistortionDisplacementMapMID);
+	SetDistortionMaterialParameters(DistortionDisplacementMapMID);
+}
+
+void UBrownConradyUDLensDistortionModelHandler::InterpretDistortionParameters()
+{
+	LensModelClass->GetDefaultObject<ULensModel>()->FromArray<FBrownConradyUDDistortionParameters>(CurrentState.DistortionInfo.Parameters, BrownConradyUDParameters);
+}
+
+FString UBrownConradyUDLensDistortionModelHandler::GetDistortionShaderPath() const
+{
+	return TEXT("/Plugin/CameraCalibrationCore/Private/BrownConradyUDDistortion.usf");
+}
+
+
+void UBrownConradyUDLensDistortionModelHandler::ExecuteDistortionShader(FRDGBuilder& GraphBuilder, const FLensDistortionState& InCurrentState, float InverseOverscan, float CameraOverscan, const FVector2D& SensorSize, FRDGTextureRef& OutDistortionMap) const
+{
+	FIntPoint DistortionMapResolution = OutDistortionMap->Desc.Extent;
+
+	FBrownConradyUDDistortionCS::FParameters* ShaderParams = GraphBuilder.AllocParameters<FBrownConradyUDDistortionCS::FParameters>();
+	
+	ShaderParams->OutDistortionMap = GraphBuilder.CreateUAV(OutDistortionMap);
+	ShaderParams->ThreadIdToUV = FVector2f(1.0f) / FVector2f(DistortionMapResolution);
+	ShaderParams->InverseOverscan = InverseOverscan;
+	ShaderParams->CameraOverscan = CameraOverscan;
+
+	// Setup Brown-Conrady U-D distortion parameters
+
+	ShaderParams->BrownConradyUDDistortionParams.ImageCenter = FVector2f(InCurrentState.ImageCenter.PrincipalPoint);
+	ShaderParams->BrownConradyUDDistortionParams.FocalLength = FVector2f(InCurrentState.FocalLengthInfo.FxFy);
+	
+	// Brown-Conrady U-D model requires exactly 8 parameters (K1, K2, K3, K4, K5, K6, P1, P2)
+	if (InCurrentState.DistortionInfo.Parameters.Num() == 8)
+	{
+		ShaderParams->BrownConradyUDDistortionParams.K1 = InCurrentState.DistortionInfo.Parameters[0];
+		ShaderParams->BrownConradyUDDistortionParams.K2 = InCurrentState.DistortionInfo.Parameters[1];
+		ShaderParams->BrownConradyUDDistortionParams.K3 = InCurrentState.DistortionInfo.Parameters[2];
+		ShaderParams->BrownConradyUDDistortionParams.K4 = InCurrentState.DistortionInfo.Parameters[3];
+		ShaderParams->BrownConradyUDDistortionParams.K5 = InCurrentState.DistortionInfo.Parameters[4];
+		ShaderParams->BrownConradyUDDistortionParams.K6 = InCurrentState.DistortionInfo.Parameters[5];
+		ShaderParams->BrownConradyUDDistortionParams.P1 = InCurrentState.DistortionInfo.Parameters[6];
+		ShaderParams->BrownConradyUDDistortionParams.P2 = InCurrentState.DistortionInfo.Parameters[7];
+	}
+	else
+	{
+		UE_LOG(LogCameraCalibrationCore, Warning, 
+			TEXT("BrownConradyUDLensDistortionModelHandler: Invalid parameter count (%d), expected exactly 8. Using default values."), 
+			InCurrentState.DistortionInfo.Parameters.Num()
+		);
+
+		ShaderParams->BrownConradyUDDistortionParams.K1 = 0.0f;
+		ShaderParams->BrownConradyUDDistortionParams.K2 = 0.0f;
+		ShaderParams->BrownConradyUDDistortionParams.K3 = 0.0f;
+		ShaderParams->BrownConradyUDDistortionParams.K4 = 0.0f;
+		ShaderParams->BrownConradyUDDistortionParams.K5 = 0.0f;
+		ShaderParams->BrownConradyUDDistortionParams.K6 = 0.0f;
+		ShaderParams->BrownConradyUDDistortionParams.P1 = 0.0f;
+		ShaderParams->BrownConradyUDDistortionParams.P2 = 0.0f;
+	}
+
+	TShaderMapRef<FBrownConradyUDDistortionCS> ComputeShader(GetGlobalShaderMap(GMaxRHIFeatureLevel));
+	FComputeShaderUtils::AddPass(
+		GraphBuilder,
+		RDG_EVENT_NAME("BrownConradyUDDistortionDisplacementMap"),
+		ComputeShader,
+		ShaderParams,
+		FIntVector(FMath::DivideAndRoundUp(DistortionMapResolution.X, 8), FMath::DivideAndRoundUp(DistortionMapResolution.Y, 8), 1));
+}
\ No newline at end of file
diff --git a/Engine/Plugins/VirtualProduction/CameraCalibrationCore/Source/CameraCalibrationCore/Private/CameraCalibrationCoreModule.cpp b/Engine/Plugins/VirtualProduction/CameraCalibrationCore/Source/CameraCalibrationCore/Private/CameraCalibrationCoreModule.cpp
--- a/Engine/Plugins/VirtualProduction/CameraCalibrationCore/Source/CameraCalibrationCore/Private/CameraCalibrationCoreModule.cpp
+++ b/Engine/Plugins/VirtualProduction/CameraCalibrationCore/Source/CameraCalibrationCore/Private/CameraCalibrationCoreModule.cpp
@@ -13,6 +13,7 @@
 #include "Misc/CoreDelegates.h"
 #include "Misc/Paths.h"
 #include "Models/AnamorphicLensModel.h"
+#include "Models/BrownConradyUDLensModel.h"
 #include "Models/SphericalLensModel.h"
 #include "Modules/ModuleManager.h"
 #include "ShaderCore.h"
@@ -121,6 +122,7 @@ void FCameraCalibrationCoreModule::RegisterDistortionModels()
 		UCameraCalibrationSubsystem* SubSystem = GEngine->GetEngineSubsystem<UCameraCalibrationSubsystem>();
 		SubSystem->RegisterDistortionModel(USphericalLensModel::StaticClass());
 		SubSystem->RegisterDistortionModel(UAnamorphicLensModel::StaticClass());
+		SubSystem->RegisterDistortionModel(UBrownConradyUDLensModel::StaticClass());
 	};
 
 	if (FApp::CanEverRender())
@@ -145,6 +147,7 @@ void FCameraCalibrationCoreModule::UnregisterDistortionModels()
 		{
 			SubSystem->UnregisterDistortionModel(USphericalLensModel::StaticClass());
 			SubSystem->UnregisterDistortionModel(UAnamorphicLensModel::StaticClass());
+			SubSystem->UnregisterDistortionModel(UBrownConradyUDLensModel::StaticClass());
 		}
 	}
 }
diff --git a/Engine/Plugins/VirtualProduction/CameraCalibrationCore/Source/CameraCalibrationCore/Private/CameraCalibrationSettings.cpp b/Engine/Plugins/VirtualProduction/CameraCalibrationCore/Source/CameraCalibrationCore/Private/CameraCalibrationSettings.cpp
--- a/Engine/Plugins/VirtualProduction/CameraCalibrationCore/Source/CameraCalibrationCore/Private/CameraCalibrationSettings.cpp
+++ b/Engine/Plugins/VirtualProduction/CameraCalibrationCore/Source/CameraCalibrationCore/Private/CameraCalibrationSettings.cpp
@@ -3,6 +3,7 @@
 #include "CameraCalibrationSettings.h"
 
 #include "AnamorphicLensDistortionModelHandler.h"
+#include "BrownConradyUDLensDistortionModelHandler.h"
 #include "CameraCalibrationSubsystem.h"
 #include "Engine/Engine.h"
 #include "Misc/CoreDelegates.h"
@@ -22,6 +23,11 @@ UCameraCalibrationSettings::UCameraCalibrationSettings()
 	DefaultDistortionDisplacementMaterials.Add(UAnamorphicLensDistortionModelHandler::StaticClass(), TSoftObjectPtr<UMaterialInterface>(FSoftObjectPath(TEXT("/CameraCalibrationCore/Materials/M_AnamorphicDistortionDisplacementMap.M_AnamorphicDistortionDisplacementMap"))));
 	DefaultDistortionMaterials.Add(UAnamorphicLensDistortionModelHandler::StaticClass(), TSoftObjectPtr<UMaterialInterface>(FSoftObjectPath(TEXT("/CameraCalibrationCore/Materials/M_DistortionPostProcess.M_DistortionPostProcess"))));
 
+	// Default materials to use for the Brown-Conrady U-D lens model
+	DefaultUndistortionDisplacementMaterials.Add(UBrownConradyUDLensDistortionModelHandler::StaticClass(), TSoftObjectPtr<UMaterialInterface>(FSoftObjectPath(TEXT("/CameraCalibrationCore/Materials/M_BrownConradyUDUndistortionDisplacementMap.M_BrownConradyUDUndistortionDisplacementMap"))));
+	DefaultDistortionDisplacementMaterials.Add(UBrownConradyUDLensDistortionModelHandler::StaticClass(), TSoftObjectPtr<UMaterialInterface>(FSoftObjectPath(TEXT("/CameraCalibrationCore/Materials/M_BrownConradyUDDistortionDisplacementMap.M_BrownConradyUDDistortionDisplacementMap"))));
+	DefaultDistortionMaterials.Add(UBrownConradyUDLensDistortionModelHandler::StaticClass(), TSoftObjectPtr<UMaterialInterface>(FSoftObjectPath(TEXT("/CameraCalibrationCore/Materials/M_DistortionPostProcess.M_DistortionPostProcess"))));
+
 #if WITH_EDITORONLY_DATA
 	auto InitOverlayOverrides = [this]()
 	{
diff --git a/Engine/Plugins/VirtualProduction/CameraCalibrationCore/Source/CameraCalibrationCore/Private/Models/BrownConradyUDLensModel.cpp b/Engine/Plugins/VirtualProduction/CameraCalibrationCore/Source/CameraCalibrationCore/Private/Models/BrownConradyUDLensModel.cpp
new file mode 100644
--- /dev/null
+++ b/Engine/Plugins/VirtualProduction/CameraCalibrationCore/Source/CameraCalibrationCore/Private/Models/BrownConradyUDLensModel.cpp
@@ -0,0 +1,20 @@
+// Copyright Epic Games, Inc. All Rights Reserved.
+
+#include "Models/BrownConradyUDLensModel.h"
+
+#include UE_INLINE_GENERATED_CPP_BY_NAME(BrownConradyUDLensModel)
+
+UScriptStruct* UBrownConradyUDLensModel::GetParameterStruct() const
+{
+	return FBrownConradyUDDistortionParameters::StaticStruct();
+}
+
+FName UBrownConradyUDLensModel::GetModelName() const 
+{ 
+	return FName("Brown-Conrady U-D Lens Model"); 
+}
+
+FName UBrownConradyUDLensModel::GetShortModelName() const
+{
+	return TEXT("BrownConradyUD");
+}
\ No newline at end of file
diff --git a/Engine/Plugins/VirtualProduction/CameraCalibrationCore/Source/CameraCalibrationCore/Public/BrownConradyUDLensDistortionModelHandler.h b/Engine/Plugins/VirtualProduction/CameraCalibrationCore/Source/CameraCalibrationCore/Public/BrownConradyUDLensDistortionModelHandler.h
new file mode 100644
--- /dev/null
+++ b/Engine/Plugins/VirtualProduction/CameraCalibrationCore/Source/CameraCalibrationCore/Public/BrownConradyUDLensDistortionModelHandler.h
@@ -0,0 +1,36 @@
+// Copyright Epic Games, Inc. All Rights Reserved.
+
+#pragma once
+
+#include "LensDistortionModelHandlerBase.h"
+
+#include "Models/BrownConradyUDLensModel.h"
+
+#include "BrownConradyUDLensDistortionModelHandler.generated.h"
+
+#define UE_API CAMERACALIBRATIONCORE_API
+
+/** Lens distortion handler for a Brown-Conrady U-D lens model that implements the Brown-Conrady polynomial model */
+UCLASS(MinimalAPI, BlueprintType)
+class UBrownConradyUDLensDistortionModelHandler : public ULensDistortionModelHandlerBase
+{
+	GENERATED_BODY()
+
+protected:
+	//~ Begin ULensDistortionModelHandlerBase interface
+	UE_API virtual void InitializeHandler() override;
+	UE_API virtual FVector2D ComputeDistortedUV(const FVector2D& InScreenUV) const override;
+	UE_API virtual FVector2D ComputeUndistortedUV(const FVector2D& InScreenUV) const override;
+	UE_API virtual void InitDistortionMaterials() override;
+	UE_API virtual void UpdateMaterialParameters() override;
+	UE_API virtual void InterpretDistortionParameters() override;
+	UE_API virtual FString GetDistortionShaderPath() const override;
+	UE_API virtual void ExecuteDistortionShader(class FRDGBuilder& GraphBuilder, const FLensDistortionState& InCurrentState, float InverseOverscan, float CameraOverscan, const FVector2D& SensorSize, FRDGTextureRef& OutDistortionMap) const override;
+	//~ End ULensDistortionModelHandlerBase interface
+
+private:
+	/** Brown-Conrady U-D lens distortion parameters (k1, k2, k3, k4, k5, k6, p1, p2) */
+	FBrownConradyUDDistortionParameters BrownConradyUDParameters;
+};
+
+#undef UE_API
\ No newline at end of file
diff --git a/Engine/Plugins/VirtualProduction/CameraCalibrationCore/Source/CameraCalibrationCore/Public/Models/BrownConradyUDLensModel.h b/Engine/Plugins/VirtualProduction/CameraCalibrationCore/Source/CameraCalibrationCore/Public/Models/BrownConradyUDLensModel.h
new file mode 100644
--- /dev/null
+++ b/Engine/Plugins/VirtualProduction/CameraCalibrationCore/Source/CameraCalibrationCore/Public/Models/BrownConradyUDLensModel.h
@@ -0,0 +1,71 @@
+// Copyright Epic Games, Inc. All Rights Reserved.
+
+#pragma once
+
+#include "LensModel.h"
+
+#include "BrownConradyUDLensModel.generated.h"
+
+#define UE_API CAMERACALIBRATIONCORE_API
+
+/**
+ * Brown-Conrady U-D lens distortion parameters (polynomial division model), using
+ * a polynomial division model: (1 + K1*r^2 + K2*r^4 + K3*r^6) / (1 + K4*r^2 + K5*r^4 + K6*r^6)
+ * where the output of that equation represents the distorted image coordinates.
+ */
+USTRUCT(BlueprintType)
+struct FBrownConradyUDDistortionParameters
+{
+	GENERATED_BODY()
+
+public:
+	/** Radial coefficient of the r^2 term (numerator) */
+	UPROPERTY(EditAnywhere, BlueprintReadWrite, Category = "Distortion")
+	float K1 = 0.0f;
+
+	/** Radial coefficient of the r^4 term (numerator) */
+	UPROPERTY(EditAnywhere, BlueprintReadWrite, Category = "Distortion")
+	float K2 = 0.0f;
+
+	/** Radial coefficient of the r^6 term (numerator) */
+	UPROPERTY(EditAnywhere, BlueprintReadWrite, Category = "Distortion")
+	float K3 = 0.0f;
+
+	/** Radial coefficient of the r^2 term (denominator) */
+	UPROPERTY(EditAnywhere, BlueprintReadWrite, Category = "Distortion")
+	float K4 = 0.0f;
+
+	/** Radial coefficient of the r^4 term (denominator) */
+	UPROPERTY(EditAnywhere, BlueprintReadWrite, Category = "Distortion")
+	float K5 = 0.0f;
+
+	/** Radial coefficient of the r^6 term (denominator) */
+	UPROPERTY(EditAnywhere, BlueprintReadWrite, Category = "Distortion")
+	float K6 = 0.0f;
+
+	/** First tangential coefficient */
+	UPROPERTY(EditAnywhere, BlueprintReadWrite, Category = "Distortion")
+	float P1 = 0.0f;
+
+	/** Second tangential coefficient */
+	UPROPERTY(EditAnywhere, BlueprintReadWrite, Category = "Distortion")
+	float P2 = 0.0f;
+};
+
+/**
+ * Brown-Conrady U-D lens model, using Brown-Conrady U-D lens distortion parameters
+ */
+UCLASS(MinimalAPI, BlueprintType, meta = (DisplayName = "Brown-Conrady U-D Lens Model"))
+class UBrownConradyUDLensModel : public ULensModel
+{
+	GENERATED_BODY()
+
+public:
+	//~ Begin ULensModel interface
+	UE_API virtual UScriptStruct* GetParameterStruct() const override;
+	UE_API virtual FName GetModelName() const override;
+	UE_API virtual FName GetShortModelName() const override;
+	//~ End ULensModel interface
+};
+
+#undef UE_API
\ No newline at end of file
