CameraCalibration: Added brown-conrady D-U model.

It subclasses the U-D model and inverts the usage of materials and LUT assignments since it goes in the opposite direction. D-U models should be better when using the PPM path, faster and more accurate because they do not need to invert the distortion equation.

[RN] Added brown-conrady D-U model. Same as U-D but in reverse. OpenTrackIO can use this model.

diff --git a/Engine/Plugins/VirtualProduction/CameraCalibrationCore/Source/CameraCalibrationCore/Private/AnamorphicLensDistortionModelHandler.cpp b/Engine/Plugins/VirtualProduction/CameraCalibrationCore/Source/CameraCalibrationCore/Private/AnamorphicLensDistortionModelHandler.cpp
--- a/Engine/Plugins/VirtualProduction/CameraCalibrationCore/Source/CameraCalibrationCore/Private/AnamorphicLensDistortionModelHandler.cpp
+++ b/Engine/Plugins/VirtualProduction/CameraCalibrationCore/Source/CameraCalibrationCore/Private/AnamorphicLensDistortionModelHandler.cpp
@@ -247,6 +247,15 @@ void UAnamorphicLensDistortionModelHandler::ExecuteDistortionShader(FRDGBuilder&
 {
 	FIntPoint DistortionMapResolution = OutDistortionMap->Desc.Extent;
 
+	// Early out if the distortion map has zero dimensions.
+	if (DistortionMapResolution.X <= 0 || DistortionMapResolution.Y <= 0)
+	{
+		UE_LOG(LogCameraCalibrationCore, Warning, 
+			TEXT("AnamorphicLensDistortionModelHandler: Skipping distortion shader execution due to invalid distortion map dimensions (%dx%d)"), 
+			DistortionMapResolution.X, DistortionMapResolution.Y);
+		return;
+	}
+
 	FAnamorphicDistortionCS::FParameters* ShaderParams = GraphBuilder.AllocParameters<FAnamorphicDistortionCS::FParameters>();
 	
 	ShaderParams->OutDistortionMap = GraphBuilder.CreateUAV(OutDistortionMap);
diff --git a/Engine/Plugins/VirtualProduction/CameraCalibrationCore/Source/CameraCalibrationCore/Private/BrownConradyDULensDistortionModelHandler.cpp b/Engine/Plugins/VirtualProduction/CameraCalibrationCore/Source/CameraCalibrationCore/Private/BrownConradyDULensDistortionModelHandler.cpp
new file mode 100644
--- /dev/null
+++ b/Engine/Plugins/VirtualProduction/CameraCalibrationCore/Source/CameraCalibrationCore/Private/BrownConradyDULensDistortionModelHandler.cpp
@@ -0,0 +1,57 @@
+// Copyright Epic Games, Inc. All Rights Reserved.
+
+#include "BrownConradyDULensDistortionModelHandler.h"
+
+#include "CameraCalibrationSettings.h"
+#include "Engine/TextureRenderTarget2D.h"
+#include "Models/BrownConradyDULensModel.h"
+
+#include UE_INLINE_GENERATED_CPP_BY_NAME(BrownConradyDULensDistortionModelHandler)
+
+void UBrownConradyDULensDistortionModelHandler::InitializeHandler()
+{
+	// Initialize with DU model instead of UD
+	LensModelClass = UBrownConradyDULensModel::StaticClass();
+}
+
+FVector2D UBrownConradyDULensDistortionModelHandler::ComputeDistortedUV(const FVector2D& InScreenUV) const
+{
+	// For DU model, ComputeDistortedUV actually performs undistortion
+	// So we call the parent's ComputeUndistortedUV
+	return Super::ComputeUndistortedUV(InScreenUV);
+}
+
+FVector2D UBrownConradyDULensDistortionModelHandler::ComputeUndistortedUV(const FVector2D& InScreenUV) const
+{
+	// For DU model, ComputeUndistortedUV actually performs distortion
+	// So we call the parent's ComputeDistortedUV
+	return Super::ComputeDistortedUV(InScreenUV);
+}
+
+void UBrownConradyDULensDistortionModelHandler::InitDistortionMaterials()
+{
+	// Initialize materials using the parent class since DU uses the same materials as UD
+	if (DistortionPostProcessMID == nullptr)
+	{
+		UMaterialInterface* DistortionMaterialParent = GetDefault<UCameraCalibrationSettings>()->GetDefaultDistortionMaterial(UBrownConradyUDLensDistortionModelHandler::StaticClass());
+		DistortionPostProcessMID = UMaterialInstanceDynamic::Create(DistortionMaterialParent, this);
+	}
+
+	if (UndistortionDisplacementMapMID == nullptr)
+	{
+		UMaterialInterface* MaterialParent = GetDefault<UCameraCalibrationSettings>()->GetDefaultUndistortionDisplacementMaterial(UBrownConradyUDLensDistortionModelHandler::StaticClass());
+		UndistortionDisplacementMapMID = UMaterialInstanceDynamic::Create(MaterialParent, this);
+	}
+
+	if (DistortionDisplacementMapMID == nullptr)
+	{
+		UMaterialInterface* MaterialParent = GetDefault<UCameraCalibrationSettings>()->GetDefaultDistortionDisplacementMaterial(UBrownConradyUDLensDistortionModelHandler::StaticClass());
+		DistortionDisplacementMapMID = UMaterialInstanceDynamic::Create(MaterialParent, this);
+	}
+
+	// For DU model, swap the texture assignments since our direction is reversed
+	DistortionPostProcessMID->SetTextureParameterValue("UndistortionDisplacementMap", DistortionDisplacementMapRT);
+	DistortionPostProcessMID->SetTextureParameterValue("DistortionDisplacementMap", UndistortionDisplacementMapRT);
+
+	SetDistortionState(CurrentState);
+}
\ No newline at end of file
diff --git a/Engine/Plugins/VirtualProduction/CameraCalibrationCore/Source/CameraCalibrationCore/Private/BrownConradyUDLensDistortionModelHandler.cpp b/Engine/Plugins/VirtualProduction/CameraCalibrationCore/Source/CameraCalibrationCore/Private/BrownConradyUDLensDistortionModelHandler.cpp
--- a/Engine/Plugins/VirtualProduction/CameraCalibrationCore/Source/CameraCalibrationCore/Private/BrownConradyUDLensDistortionModelHandler.cpp
+++ b/Engine/Plugins/VirtualProduction/CameraCalibrationCore/Source/CameraCalibrationCore/Private/BrownConradyUDLensDistortionModelHandler.cpp
@@ -205,6 +205,15 @@ void UBrownConradyUDLensDistortionModelHandler::ExecuteDistortionShader(FRDGBuil
 {
 	FIntPoint DistortionMapResolution = OutDistortionMap->Desc.Extent;
 
+	// Early out if the distortion map has zero dimensions.
+	if (DistortionMapResolution.X <= 0 || DistortionMapResolution.Y <= 0)
+	{
+		UE_LOG(LogCameraCalibrationCore, Warning, 
+			TEXT("BrownConradyUDLensDistortionModelHandler: Skipping distortion shader execution due to invalid distortion map dimensions (%dx%d)"), 
+			DistortionMapResolution.X, DistortionMapResolution.Y);
+		return;
+	}
+
 	FBrownConradyUDDistortionCS::FParameters* ShaderParams = GraphBuilder.AllocParameters<FBrownConradyUDDistortionCS::FParameters>();
 	
 	ShaderParams->OutDistortionMap = GraphBuilder.CreateUAV(OutDistortionMap);
diff --git a/Engine/Plugins/VirtualProduction/CameraCalibrationCore/Source/CameraCalibrationCore/Private/CameraCalibrationCoreModule.cpp b/Engine/Plugins/VirtualProduction/CameraCalibrationCore/Source/CameraCalibrationCore/Private/CameraCalibrationCoreModule.cpp
--- a/Engine/Plugins/VirtualProduction/CameraCalibrationCore/Source/CameraCalibrationCore/Private/CameraCalibrationCoreModule.cpp
+++ b/Engine/Plugins/VirtualProduction/CameraCalibrationCore/Source/CameraCalibrationCore/Private/CameraCalibrationCoreModule.cpp
@@ -13,6 +13,7 @@
 #include "Misc/CoreDelegates.h"
 #include "Misc/Paths.h"
 #include "Models/AnamorphicLensModel.h"
+#include "Models/BrownConradyDULensModel.h"
 #include "Models/BrownConradyUDLensModel.h"
 #include "Models/SphericalLensModel.h"
 #include "Modules/ModuleManager.h"
@@ -123,6 +124,7 @@ void FCameraCalibrationCoreModule::RegisterDistortionModels()
 		SubSystem->RegisterDistortionModel(USphericalLensModel::StaticClass());
 		SubSystem->RegisterDistortionModel(UAnamorphicLensModel::StaticClass());
 		SubSystem->RegisterDistortionModel(UBrownConradyUDLensModel::StaticClass());
+		SubSystem->RegisterDistortionModel(UBrownConradyDULensModel::StaticClass());
 	};
 
 	if (FApp::CanEverRender())
@@ -148,6 +150,7 @@ void FCameraCalibrationCoreModule::UnregisterDistortionModels()
 			SubSystem->UnregisterDistortionModel(USphericalLensModel::StaticClass());
 			SubSystem->UnregisterDistortionModel(UAnamorphicLensModel::StaticClass());
 			SubSystem->UnregisterDistortionModel(UBrownConradyUDLensModel::StaticClass());
+			SubSystem->UnregisterDistortionModel(UBrownConradyDULensModel::StaticClass());
 		}
 	}
 }
diff --git a/Engine/Plugins/VirtualProduction/CameraCalibrationCore/Source/CameraCalibrationCore/Private/Models/BrownConradyDULensModel.cpp b/Engine/Plugins/VirtualProduction/CameraCalibrationCore/Source/CameraCalibrationCore/Private/Models/BrownConradyDULensModel.cpp
new file mode 100644
--- /dev/null
+++ b/Engine/Plugins/VirtualProduction/CameraCalibrationCore/Source/CameraCalibrationCore/Private/Models/BrownConradyDULensModel.cpp
@@ -0,0 +1,21 @@
+// Copyright Epic Games, Inc. All Rights Reserved.
+
+#include "Models/BrownConradyDULensModel.h"
+
+#include UE_INLINE_GENERATED_CPP_BY_NAME(BrownConradyDULensModel)
+
+UScriptStruct* UBrownConradyDULensModel::GetParameterStruct() const
+{
+	// Use the same parameter struct as UD model
+	return FBrownConradyUDDistortionParameters::StaticStruct();
+}
+
+FName UBrownConradyDULensModel::GetModelName() const
+{
+	return FName("Brown-Conrady D-U Lens Model");
+}
+
+FName UBrownConradyDULensModel::GetShortModelName() const
+{
+	return FName("BrownConradyDU");
+}
\ No newline at end of file
diff --git a/Engine/Plugins/VirtualProduction/CameraCalibrationCore/Source/CameraCalibrationCore/Private/SphericalLensDistortionModelHandler.cpp b/Engine/Plugins/VirtualProduction/CameraCalibrationCore/Source/CameraCalibrationCore/Private/SphericalLensDistortionModelHandler.cpp
--- a/Engine/Plugins/VirtualProduction/CameraCalibrationCore/Source/CameraCalibrationCore/Private/SphericalLensDistortionModelHandler.cpp
+++ b/Engine/Plugins/VirtualProduction/CameraCalibrationCore/Source/CameraCalibrationCore/Private/SphericalLensDistortionModelHandler.cpp
@@ -185,6 +185,15 @@ void USphericalLensDistortionModelHandler::ExecuteDistortionShader(FRDGBuilder&
 {
 	FIntPoint DistortionMapResolution = OutDistortionMap->Desc.Extent;
 
+	// Early out if the distortion map has zero dimensions.
+	if (DistortionMapResolution.X <= 0 || DistortionMapResolution.Y <= 0)
+	{
+		UE_LOG(LogCameraCalibrationCore, Warning, 
+			TEXT("SphericalLensDistortionModelHandler: Skipping distortion shader execution due to invalid distortion map dimensions (%dx%d)"), 
+			DistortionMapResolution.X, DistortionMapResolution.Y);
+		return;
+	}
+
 	FSphericalDistortionCS::FParameters* ShaderParams = GraphBuilder.AllocParameters<FSphericalDistortionCS::FParameters>();
 	
 	ShaderParams->OutDistortionMap = GraphBuilder.CreateUAV(OutDistortionMap);
diff --git a/Engine/Plugins/VirtualProduction/CameraCalibrationCore/Source/CameraCalibrationCore/Public/BrownConradyDULensDistortionModelHandler.h b/Engine/Plugins/VirtualProduction/CameraCalibrationCore/Source/CameraCalibrationCore/Public/BrownConradyDULensDistortionModelHandler.h
new file mode 100644
--- /dev/null
+++ b/Engine/Plugins/VirtualProduction/CameraCalibrationCore/Source/CameraCalibrationCore/Public/BrownConradyDULensDistortionModelHandler.h
@@ -0,0 +1,31 @@
+// Copyright Epic Games, Inc. All Rights Reserved.
+
+#pragma once
+
+#include "BrownConradyUDLensDistortionModelHandler.h"
+
+#include "BrownConradyDULensDistortionModelHandler.generated.h"
+
+#define UE_API CAMERACALIBRATIONCORE_API
+
+/** 
+ * Lens distortion handler for a Brown-Conrady D-U lens model 
+ * This model interprets the distortion parameters in the reverse direction (distorted to undistorted)
+ * compared to the U-D model. It subclasses the U-D handler and overrides key functions to swap directions.
+ */
+UCLASS(MinimalAPI, BlueprintType)
+class UBrownConradyDULensDistortionModelHandler : public UBrownConradyUDLensDistortionModelHandler
+{
+	GENERATED_BODY()
+
+protected:
+	//~ Begin ULensDistortionModelHandlerBase interface
+	UE_API virtual void InitializeHandler() override;
+	UE_API virtual bool IsForwardDistorting() const override { return false; }
+	UE_API virtual FVector2D ComputeDistortedUV(const FVector2D& InScreenUV) const override;
+	UE_API virtual FVector2D ComputeUndistortedUV(const FVector2D& InScreenUV) const override;
+	UE_API virtual void InitDistortionMaterials() override;
+	//~ End ULensDistortionModelHandlerBase interface
+};
+
+#undef UE_API
\ No newline at end of file
diff --git a/Engine/Plugins/VirtualProduction/CameraCalibrationCore/Source/CameraCalibrationCore/Public/Models/BrownConradyDULensModel.h b/Engine/Plugins/VirtualProduction/CameraCalibrationCore/Source/CameraCalibrationCore/Public/Models/BrownConradyDULensModel.h
new file mode 100644
--- /dev/null
+++ b/Engine/Plugins/VirtualProduction/CameraCalibrationCore/Source/CameraCalibrationCore/Public/Models/BrownConradyDULensModel.h
@@ -0,0 +1,30 @@
+// Copyright Epic Games, Inc. All Rights Reserved.
+
+#pragma once
+
+#include "LensModel.h"
+#include "BrownConradyUDLensModel.h"
+
+#include "BrownConradyDULensModel.generated.h"
+
+#define UE_API CAMERACALIBRATIONCORE_API
+
+/**
+ * Brown-Conrady D-U lens model (Distorted to Undistorted direction)
+ * Uses the same parameters as Brown-Conrady U-D but interprets them in the opposite direction.
+ * The polynomial division model represents the mapping from distorted to undistorted coordinates.
+ */
+UCLASS(MinimalAPI, BlueprintType, meta = (DisplayName = "Brown-Conrady D-U Lens Model"))
+class UBrownConradyDULensModel : public ULensModel
+{
+	GENERATED_BODY()
+
+public:
+	//~ Begin ULensModel interface
+	UE_API virtual UScriptStruct* GetParameterStruct() const override;
+	UE_API virtual FName GetModelName() const override;
+	UE_API virtual FName GetShortModelName() const override;
+	//~ End ULensModel interface
+};
+
+#undef UE_API
\ No newline at end of file
