CameraCalibration: Fixed inaccurate comment about nature of early out.

#rnx

diff --git a/Engine/Plugins/VirtualProduction/CameraCalibrationCore/Source/CameraCalibrationCore/Private/AnamorphicLensDistortionModelHandler.cpp b/Engine/Plugins/VirtualProduction/CameraCalibrationCore/Source/CameraCalibrationCore/Private/AnamorphicLensDistortionModelHandler.cpp
--- a/Engine/Plugins/VirtualProduction/CameraCalibrationCore/Source/CameraCalibrationCore/Private/AnamorphicLensDistortionModelHandler.cpp
+++ b/Engine/Plugins/VirtualProduction/CameraCalibrationCore/Source/CameraCalibrationCore/Private/AnamorphicLensDistortionModelHandler.cpp
@@ -247,7 +247,7 @@ void UAnamorphicLensDistortionModelHandler::ExecuteDistortionShader(FRDGBuilder&
 {
 	FIntPoint DistortionMapResolution = OutDistortionMap->Desc.Extent;
 
-	// Early out if the distortion map has zero dimensions.
+	// Early out if the distortion map has invalid dimensions.
 	if (DistortionMapResolution.X <= 0 || DistortionMapResolution.Y <= 0)
 	{
 		UE_LOG(LogCameraCalibrationCore, Warning, 
diff --git a/Engine/Plugins/VirtualProduction/CameraCalibrationCore/Source/CameraCalibrationCore/Private/BrownConradyUDLensDistortionModelHandler.cpp b/Engine/Plugins/VirtualProduction/CameraCalibrationCore/Source/CameraCalibrationCore/Private/BrownConradyUDLensDistortionModelHandler.cpp
--- a/Engine/Plugins/VirtualProduction/CameraCalibrationCore/Source/CameraCalibrationCore/Private/BrownConradyUDLensDistortionModelHandler.cpp
+++ b/Engine/Plugins/VirtualProduction/CameraCalibrationCore/Source/CameraCalibrationCore/Private/BrownConradyUDLensDistortionModelHandler.cpp
@@ -205,7 +205,7 @@ void UBrownConradyUDLensDistortionModelHandler::ExecuteDistortionShader(FRDGBuil
 {
 	FIntPoint DistortionMapResolution = OutDistortionMap->Desc.Extent;
 
-	// Early out if the distortion map has zero dimensions.
+	// Early out if the distortion map has invalid dimensions.
 	if (DistortionMapResolution.X <= 0 || DistortionMapResolution.Y <= 0)
 	{
 		UE_LOG(LogCameraCalibrationCore, Warning, 
diff --git a/Engine/Plugins/VirtualProduction/CameraCalibrationCore/Source/CameraCalibrationCore/Private/SphericalLensDistortionModelHandler.cpp b/Engine/Plugins/VirtualProduction/CameraCalibrationCore/Source/CameraCalibrationCore/Private/SphericalLensDistortionModelHandler.cpp
--- a/Engine/Plugins/VirtualProduction/CameraCalibrationCore/Source/CameraCalibrationCore/Private/SphericalLensDistortionModelHandler.cpp
+++ b/Engine/Plugins/VirtualProduction/CameraCalibrationCore/Source/CameraCalibrationCore/Private/SphericalLensDistortionModelHandler.cpp
@@ -185,7 +185,7 @@ void USphericalLensDistortionModelHandler::ExecuteDistortionShader(FRDGBuilder&
 {
 	FIntPoint DistortionMapResolution = OutDistortionMap->Desc.Extent;
 
-	// Early out if the distortion map has zero dimensions.
+	// Early out if the distortion map has invalid dimensions.
 	if (DistortionMapResolution.X <= 0 || DistortionMapResolution.Y <= 0)
 	{
 		UE_LOG(LogCameraCalibrationCore, Warning, 
