CameraCalibrationCore: Fix support for packaged game.

diff --git a/Engine/Plugins/VirtualProduction/CameraCalibrationCore/Config/Game.ini b/Engine/Plugins/VirtualProduction/CameraCalibrationCore/Config/Game.ini
new file mode 100644
--- /dev/null
+++ b/Engine/Plugins/VirtualProduction/CameraCalibrationCore/Config/Game.ini
@@ -0,0 +1,2 @@
+[/Script/UnrealEd.ProjectPackagingSettings]
++DirectoriesToAlwaysCook=(Path="/CameraCalibrationCore")
\ No newline at end of file
