OpenTrackIO: Fix projection offset in BrownConrady U-D and D-U distortion models.

When distortion offset was added, it incorrectly replaced projection offset, which is the more commonly needed calibration offset.

Also fixed related regression in the simpler spherical model that also started treating the center offset as a distortion offset instead of a projection offset.

#rnx

diff --git a/Engine/Plugins/VirtualProduction/CameraCalibrationCore/Shaders/Private/BrownConradyUDDistortion.usf b/Engine/Plugins/VirtualProduction/CameraCalibrationCore/Shaders/Private/BrownConradyUDDistortion.usf
--- a/Engine/Plugins/VirtualProduction/CameraCalibrationCore/Shaders/Private/BrownConradyUDDistortion.usf
+++ b/Engine/Plugins/VirtualProduction/CameraCalibrationCore/Shaders/Private/BrownConradyUDDistortion.usf
@@ -13,6 +13,7 @@ float2 ThreadIdToUV;
 float InverseOverscan;
 float CameraOverscan;
 RWTexture2D<float2> OutDistortionMap;
+int bIsUD;
 
 // Brown-Conrady U-D distortion parameters (polynomial division model)
 float2 FocalLength;
@@ -28,8 +29,13 @@ float P2;
 
 float2 BrownConradyUDDistort(float2 InputUV)
 {
-	float x_prime = (InputUV.x - ImageCenter.x) / FocalLength.x;
-	float y_prime = (InputUV.y - ImageCenter.y) / FocalLength.y;
+	// @todo Distortion offset parameter to replace the 0.5 at the beginning and at the end.
+
+	float input_offset_x = bIsUD ? 0.5 : ImageCenter.x;
+	float input_offset_y = bIsUD ? 0.5 : ImageCenter.y;
+	
+	float x_prime = (InputUV.x - input_offset_x) / FocalLength.x;
+	float y_prime = (InputUV.y - input_offset_y) / FocalLength.y;
 
 	float x_prime2 = x_prime * x_prime;
 	float y_prime2 = y_prime * y_prime;
@@ -51,8 +57,11 @@ float2 BrownConradyUDDistort(float2 InputUV)
 	float x_double_prime = (x_prime * dr) + dtx;
 	float y_double_prime = (y_prime * dr) + dty;
 
-	float x_undistort = (FocalLength.x * x_double_prime) + ImageCenter.x;
-	float y_undistort = (FocalLength.y * y_double_prime) + ImageCenter.y;
+	float output_offset_x = bIsUD ? ImageCenter.x : 0.5;
+	float output_offset_y = bIsUD ? ImageCenter.y : 0.5;
+
+	float x_undistort = (FocalLength.x * x_double_prime) + output_offset_x;
+	float y_undistort = (FocalLength.y * y_double_prime) + output_offset_y;
 
 	float2 result;
 	result.x = x_undistort;
diff --git a/Engine/Plugins/VirtualProduction/CameraCalibrationCore/Shaders/Private/SphericalDistortion.usf b/Engine/Plugins/VirtualProduction/CameraCalibrationCore/Shaders/Private/SphericalDistortion.usf
--- a/Engine/Plugins/VirtualProduction/CameraCalibrationCore/Shaders/Private/SphericalDistortion.usf
+++ b/Engine/Plugins/VirtualProduction/CameraCalibrationCore/Shaders/Private/SphericalDistortion.usf
@@ -24,8 +24,10 @@ float P2;
 
 float2 SphericalDistort(float2 InputUV)
 {
-	float x_prime = (InputUV.x - ImageCenter.x) / FocalLength.x;
-	float y_prime = (InputUV.y - ImageCenter.y) / FocalLength.y;
+	// @todo Distortion offset to replace the 0.5 below
+
+	float x_prime = (InputUV.x - 0.5) / FocalLength.x;
+	float y_prime = (InputUV.y - 0.5) / FocalLength.y;
 
 	float x_prime2 = x_prime * x_prime;
 	float y_prime2 = y_prime * y_prime;
diff --git a/Engine/Plugins/VirtualProduction/CameraCalibrationCore/Source/CameraCalibrationCore/Private/BrownConradyDULensDistortionModelHandler.cpp b/Engine/Plugins/VirtualProduction/CameraCalibrationCore/Source/CameraCalibrationCore/Private/BrownConradyDULensDistortionModelHandler.cpp
--- a/Engine/Plugins/VirtualProduction/CameraCalibrationCore/Source/CameraCalibrationCore/Private/BrownConradyDULensDistortionModelHandler.cpp
+++ b/Engine/Plugins/VirtualProduction/CameraCalibrationCore/Source/CameraCalibrationCore/Private/BrownConradyDULensDistortionModelHandler.cpp
@@ -30,28 +30,27 @@ FVector2D UBrownConradyDULensDistortionModelHandler::ComputeUndistortedUV(const
 
 void UBrownConradyDULensDistortionModelHandler::InitDistortionMaterials()
 {
-	// Initialize materials using the parent class since DU uses the same materials as UD
+	// Initialize materials
 	if (DistortionPostProcessMID == nullptr)
 	{
-		UMaterialInterface* DistortionMaterialParent = GetDefault<UCameraCalibrationSettings>()->GetDefaultDistortionMaterial(UBrownConradyUDLensDistortionModelHandler::StaticClass());
+		UMaterialInterface* DistortionMaterialParent = GetDefault<UCameraCalibrationSettings>()->GetDefaultDistortionMaterial(UBrownConradyDULensDistortionModelHandler::StaticClass());
 		DistortionPostProcessMID = UMaterialInstanceDynamic::Create(DistortionMaterialParent, this);
 	}
 
 	if (UndistortionDisplacementMapMID == nullptr)
 	{
-		UMaterialInterface* MaterialParent = GetDefault<UCameraCalibrationSettings>()->GetDefaultUndistortionDisplacementMaterial(UBrownConradyUDLensDistortionModelHandler::StaticClass());
+		UMaterialInterface* MaterialParent = GetDefault<UCameraCalibrationSettings>()->GetDefaultUndistortionDisplacementMaterial(UBrownConradyDULensDistortionModelHandler::StaticClass());
 		UndistortionDisplacementMapMID = UMaterialInstanceDynamic::Create(MaterialParent, this);
 	}
 
 	if (DistortionDisplacementMapMID == nullptr)
 	{
-		UMaterialInterface* MaterialParent = GetDefault<UCameraCalibrationSettings>()->GetDefaultDistortionDisplacementMaterial(UBrownConradyUDLensDistortionModelHandler::StaticClass());
+		UMaterialInterface* MaterialParent = GetDefault<UCameraCalibrationSettings>()->GetDefaultDistortionDisplacementMaterial(UBrownConradyDULensDistortionModelHandler::StaticClass());
 		DistortionDisplacementMapMID = UMaterialInstanceDynamic::Create(MaterialParent, this);
 	}
 
-	// For DU model, swap the texture assignments since our direction is reversed
-	DistortionPostProcessMID->SetTextureParameterValue("UndistortionDisplacementMap", DistortionDisplacementMapRT);
-	DistortionPostProcessMID->SetTextureParameterValue("DistortionDisplacementMap", UndistortionDisplacementMapRT);
+	DistortionPostProcessMID->SetTextureParameterValue("UndistortionDisplacementMap", UndistortionDisplacementMapRT);
+	DistortionPostProcessMID->SetTextureParameterValue("DistortionDisplacementMap", DistortionDisplacementMapRT);
 
 	SetDistortionState(CurrentState);
 }
\ No newline at end of file
diff --git a/Engine/Plugins/VirtualProduction/CameraCalibrationCore/Source/CameraCalibrationCore/Private/BrownConradyUDLensDistortionModelHandler.cpp b/Engine/Plugins/VirtualProduction/CameraCalibrationCore/Source/CameraCalibrationCore/Private/BrownConradyUDLensDistortionModelHandler.cpp
--- a/Engine/Plugins/VirtualProduction/CameraCalibrationCore/Source/CameraCalibrationCore/Private/BrownConradyUDLensDistortionModelHandler.cpp
+++ b/Engine/Plugins/VirtualProduction/CameraCalibrationCore/Source/CameraCalibrationCore/Private/BrownConradyUDLensDistortionModelHandler.cpp
@@ -40,6 +40,7 @@ class FBrownConradyUDDistortionCS : public FGlobalShader
 		SHADER_PARAMETER(float, CameraOverscan)
 		SHADER_PARAMETER_STRUCT_INCLUDE(FBrownConradyUDDistortionParams, BrownConradyUDDistortionParams)
 		SHADER_PARAMETER_RDG_TEXTURE_UAV(RWTexture2D, OutDistortionMap)
+		SHADER_PARAMETER(int32, bIsUD)
 	END_SHADER_PARAMETER_STRUCT()
 
 	static bool ShouldCompilePermutation(const FGlobalShaderPermutationParameters& Parameters)
@@ -65,7 +66,8 @@ FVector2D UBrownConradyUDLensDistortionModelHandler::ComputeDistortedUV(const FV
 		return InScreenUV;
 	}
 
-	FVector2D X = (InScreenUV - CxCy) / FxFy; // normalized coords (undistorted)
+	// @todo distortion offset to replace the 0.5 below
+	FVector2D X = (InScreenUV - FVector2D(0.5f, 0.5f)) / FxFy; // normalized coords (undistorted)
 
 	const float R2 = X.X * X.X + X.Y * X.Y;
 	const float R4 = R2 * R2;
@@ -136,7 +138,9 @@ FVector2D UBrownConradyUDLensDistortionModelHandler::ComputeUndistortedUV(const
 
 		P = (Pd - Tangential) * InvRadial;
 	}
-	return P * FxFy + CxCy; // undistorted UV
+
+	// @todo distortion offset to replace the 0.5 below
+	return P * FxFy + FVector2D(0.5f, 0.5f); // undistorted UV
 }
 
 void UBrownConradyUDLensDistortionModelHandler::InitDistortionMaterials()
@@ -225,6 +229,7 @@ void UBrownConradyUDLensDistortionModelHandler::ExecuteDistortionShader(FRDGBuil
 
 	ShaderParams->BrownConradyUDDistortionParams.ImageCenter = FVector2f(InCurrentState.ImageCenter.PrincipalPoint);
 	ShaderParams->BrownConradyUDDistortionParams.FocalLength = FVector2f(InCurrentState.FocalLengthInfo.FxFy);
+	ShaderParams->bIsUD = (GetClass() == UBrownConradyUDLensDistortionModelHandler::StaticClass());
 	
 	// Brown-Conrady U-D model requires exactly 8 parameters (K1, K2, K3, K4, K5, K6, P1, P2)
 	if (InCurrentState.DistortionInfo.Parameters.Num() == 8)
diff --git a/Engine/Plugins/VirtualProduction/CameraCalibrationCore/Source/CameraCalibrationCore/Private/CameraCalibrationSettings.cpp b/Engine/Plugins/VirtualProduction/CameraCalibrationCore/Source/CameraCalibrationCore/Private/CameraCalibrationSettings.cpp
--- a/Engine/Plugins/VirtualProduction/CameraCalibrationCore/Source/CameraCalibrationCore/Private/CameraCalibrationSettings.cpp
+++ b/Engine/Plugins/VirtualProduction/CameraCalibrationCore/Source/CameraCalibrationCore/Private/CameraCalibrationSettings.cpp
@@ -3,6 +3,7 @@
 #include "CameraCalibrationSettings.h"
 
 #include "AnamorphicLensDistortionModelHandler.h"
+#include "BrownConradyDULensDistortionModelHandler.h"
 #include "BrownConradyUDLensDistortionModelHandler.h"
 #include "CameraCalibrationSubsystem.h"
 #include "Engine/Engine.h"
@@ -28,6 +29,11 @@ UCameraCalibrationSettings::UCameraCalibrationSettings()
 	DefaultDistortionDisplacementMaterials.Add(UBrownConradyUDLensDistortionModelHandler::StaticClass(), TSoftObjectPtr<UMaterialInterface>(FSoftObjectPath(TEXT("/CameraCalibrationCore/Materials/M_BrownConradyUDDistortionDisplacementMap.M_BrownConradyUDDistortionDisplacementMap"))));
 	DefaultDistortionMaterials.Add(UBrownConradyUDLensDistortionModelHandler::StaticClass(), TSoftObjectPtr<UMaterialInterface>(FSoftObjectPath(TEXT("/CameraCalibrationCore/Materials/M_DistortionPostProcess.M_DistortionPostProcess"))));
 
+	// Default materials to use for the Brown-Conrady D-U lens model
+	DefaultUndistortionDisplacementMaterials.Add(UBrownConradyDULensDistortionModelHandler::StaticClass(), TSoftObjectPtr<UMaterialInterface>(FSoftObjectPath(TEXT("/CameraCalibrationCore/Materials/M_BrownConradyDUUndistortionDisplacementMap.M_BrownConradyDUUndistortionDisplacementMap"))));
+	DefaultDistortionDisplacementMaterials.Add(UBrownConradyDULensDistortionModelHandler::StaticClass(), TSoftObjectPtr<UMaterialInterface>(FSoftObjectPath(TEXT("/CameraCalibrationCore/Materials/M_BrownConradyDUDistortionDisplacementMap.M_BrownConradyDUDistortionDisplacementMap"))));
+	DefaultDistortionMaterials.Add(UBrownConradyDULensDistortionModelHandler::StaticClass(), TSoftObjectPtr<UMaterialInterface>(FSoftObjectPath(TEXT("/CameraCalibrationCore/Materials/M_DistortionPostProcess.M_DistortionPostProcess"))));
+
 #if WITH_EDITORONLY_DATA
 	auto InitOverlayOverrides = [this]()
 	{
diff --git a/Engine/Plugins/VirtualProduction/CameraCalibrationCore/Source/CameraCalibrationCore/Private/SphericalLensDistortionModelHandler.cpp b/Engine/Plugins/VirtualProduction/CameraCalibrationCore/Source/CameraCalibrationCore/Private/SphericalLensDistortionModelHandler.cpp
--- a/Engine/Plugins/VirtualProduction/CameraCalibrationCore/Source/CameraCalibrationCore/Private/SphericalLensDistortionModelHandler.cpp
+++ b/Engine/Plugins/VirtualProduction/CameraCalibrationCore/Source/CameraCalibrationCore/Private/SphericalLensDistortionModelHandler.cpp
@@ -62,7 +62,8 @@ FVector2D USphericalLensDistortionModelHandler::ComputeDistortedUV(const FVector
 		return InScreenUV;
 	}
 
-	FVector2D X = (InScreenUV - CxCy) / FxFy; // normalized coords (undistorted)
+	// @todo distortion offset to replace the 0.5 below
+	FVector2D X = (InScreenUV - FVector2D(0.5f, 0.5f)) / FxFy; // normalized coords (undistorted)
 
 	const float R2 = X.X * X.X + X.Y * X.Y;
 	const float R4 = R2 * R2;
@@ -119,7 +120,9 @@ FVector2D USphericalLensDistortionModelHandler::ComputeUndistortedUV(const FVect
 
 		P = (Pd - Tangential) * InvRadial;
 	}
-	return P * FxFy + CxCy; // undistorted UV
+
+	// @todo distortion offset to replace the 0.5 below
+	return P * FxFy + FVector2D(0.5f, 0.5f); // undistorted UV
 }
 
 void USphericalLensDistortionModelHandler::InitDistortionMaterials()
diff --git a/Engine/Plugins/VirtualProduction/LiveLinkOpenTrackIO/Source/LiveLinkOpenTrackIO/Private/LiveLinkOpenTrackIOConversions.cpp b/Engine/Plugins/VirtualProduction/LiveLinkOpenTrackIO/Source/LiveLinkOpenTrackIO/Private/LiveLinkOpenTrackIOConversions.cpp
--- a/Engine/Plugins/VirtualProduction/LiveLinkOpenTrackIO/Source/LiveLinkOpenTrackIO/Private/LiveLinkOpenTrackIOConversions.cpp
+++ b/Engine/Plugins/VirtualProduction/LiveLinkOpenTrackIO/Source/LiveLinkOpenTrackIO/Private/LiveLinkOpenTrackIOConversions.cpp
@@ -129,11 +129,11 @@ namespace LiveLinkOpenTrackIOConversions
 			{
 				// Cx Cy
 				{
-					// OpenTrackIO DistortionOffset is in mm, Unreal PrincipalPoint is normalized 0..1, centered at 0.5.
-					// We normalize DistortionOffset in mm by the filmback dimensions that are also in mm.
+					// OpenTrackIO ProjectionOffset is in mm, Unreal PrincipalPoint is normalized 0..1, centered at 0.5.
+					// We normalize ProjectionOffset in mm by the filmback dimensions that are also in mm.
 
-					OutUnrealLensData.PrincipalPoint.X = 0.5 + InLensData->DistortionOffset.X / Width_mm;
-					OutUnrealLensData.PrincipalPoint.Y = 0.5 + InLensData->DistortionOffset.Y / Height_mm;
+					OutUnrealLensData.PrincipalPoint.X = 0.5 + InLensData->ProjectionOffset.X / Width_mm;
+					OutUnrealLensData.PrincipalPoint.Y = 0.5 + InLensData->ProjectionOffset.Y / Height_mm;
 				}
 
 				// FxFy
