Camera Calibration: Fixed issue where certain types of lens files would cause a crash due to an abundance of undefined calibration points. The lens distortion interpolation now supports cases where any number of points near the interpolation are undefined as it now uses the whole focus and zoom curves to compute proper interpolation values.

#Synced-CL 50839866

#Synced-CL 50839866

diff --git a/Engine/Plugins/VirtualProduction/CameraCalibrationCore/Source/CameraCalibrationCore/Private/LensFile.cpp b/Engine/Plugins/VirtualProduction/CameraCalibrationCore/Source/CameraCalibrationCore/Private/LensFile.cpp
--- a/Engine/Plugins/VirtualProduction/CameraCalibrationCore/Source/CameraCalibrationCore/Private/LensFile.cpp
+++ b/Engine/Plugins/VirtualProduction/CameraCalibrationCore/Source/CameraCalibrationCore/Private/LensFile.cpp
@@ -264,24 +264,14 @@ bool ULensFile::EvaluateDistortionParameters(float InFocus, float InZoom, FDisto
 	using FBlendResults = LensInterpolationUtils::FDistortionMapBlendResults;
 	
 	FBlendParams Params;
-	
-	Params.GetDistortionParameters = FBlendParams::FGetDistortionParameters::CreateLambda(
-		[](const FDistortionFocusPoint& FocusPoint, const FDistortionFocusCurve& FocusCurve)
-		{
-			FDistortionInfo Point;
-			if (FocusPoint.GetPoint(FocusCurve.Zoom, Point))
-			{
-				return TOptional(Point);
-			}
-			
-			return TOptional<FDistortionInfo>();
-		}
-	);
-	
-	Params.DistortionParamNum = OutEvaluatedValue.Parameters.Num();
+	Params.bGenerateBlendedDistortionParams = true;
+	Params.GetInterpolatedImageCenter = FBlendParams::FGetInterpolatedImageCenter::CreateUObject(this, &ULensFile::EvaluateImageCenterParameters);
+	Params.GetInterpolatedFocalLength = FBlendParams::FGetInterpolatedFocalLength::CreateUObject(this, &ULensFile::EvaluateFocalLength);
+
+	LensInfo.LensModel->GetDefaultObject<ULensModel>()->GetDefaultParameterArray(Params.DefaultDistortionParams);
 
 	FBlendResults Results = LensInterpolationUtils::DistortionMapBlend(DistortionTable, InFocus, InZoom, Params);
-	if (!Results.bValid)
+	if (!Results.bValid || !Results.BlendedDistortionParams.IsSet())
 	{
 		return false;
 	}
@@ -409,28 +399,17 @@ void ULensFile::SetupNoDistortionOutput(ULensDistortionModelHandlerBase* LensHan
 
 void ULensFile::GetBlendState(float InFocus, float InZoom, FVector2D InFilmback, FDisplacementMapBlendingParams& OutBlendState)
 {
-	LensInterpolationUtils::FDistortionMapBlendParams<FDistortionTable> Params;
+	using FBlendParams = LensInterpolationUtils::FDistortionMapBlendParams<FDistortionTable>;
+	using FBlendResults = LensInterpolationUtils::FDistortionMapBlendResults;
+	
+	FBlendParams Params;
 	Params.bGenerateBlendingParams = true;
+	Params.GetInterpolatedImageCenter = FBlendParams::FGetInterpolatedImageCenter::CreateUObject(this, &ULensFile::EvaluateImageCenterParameters);
+	Params.GetInterpolatedFocalLength = FBlendParams::FGetInterpolatedFocalLength::CreateUObject(this, &ULensFile::EvaluateFocalLength);
 
- 	FImageCenterInfo InterpolatedImageCenter;
- 	EvaluateImageCenterParameters(InFocus, InZoom, InterpolatedImageCenter);
-
-	Params.GetDistortionState = LensInterpolationUtils::FDistortionMapBlendParams<FDistortionTable>::FGetDistortionState::CreateLambda(
-		[this, &InterpolatedImageCenter]
-		(const FDistortionFocusPoint& FocusPoint, const FDistortionFocusCurve& FocusCurve, FLensDistortionState& OutState)
-		{
-			// In case the point doesn't exist, we need to fill the distortion parameter array with default values
-			LensInfo.LensModel->GetDefaultObject<ULensModel>()->GetDefaultParameterArray(OutState.DistortionInfo.Parameters);
-			FocusPoint.GetPoint(FocusCurve.Zoom, OutState.DistortionInfo);
-
-			LensDataTableUtils::GetPointValue<FFocalLengthFocusPoint>(FocusPoint.Focus, FocusCurve.Zoom, FocalLengthTable.FocusPoints, OutState.FocalLengthInfo);
-			OutState.ImageCenter = InterpolatedImageCenter;
-		}
-	);
-
-	LensInterpolationUtils::FDistortionMapBlendResults Results = LensInterpolationUtils::DistortionMapBlend(DistortionTable, InFocus, InZoom, Params);
+	FBlendResults Results = LensInterpolationUtils::DistortionMapBlend(DistortionTable, InFocus, InZoom, Params);
 
-	if (Results.BlendingParams.IsSet())
+	if (Results.bValid && Results.BlendingParams.IsSet())
 	{
 		OutBlendState = Results.BlendingParams.GetValue();
 		OutBlendState.FxFyScale = FVector2D(InFilmback.X / LensInfo.SensorDimensions.X, InFilmback.Y / LensInfo.SensorDimensions.Y);
@@ -439,12 +418,8 @@ void ULensFile::GetBlendState(float InFocus, float InZoom, FVector2D InFilmback,
 
 bool ULensFile::EvaluateDistortionForParameters(float InFocus, float InZoom, FVector2D InFilmback, ULensDistortionModelHandlerBase* InLensHandler) const
 {
-	// Compute interpolated image center and focal length to pass to the lens handler
-	FImageCenterInfo InterpolatedImageCenter;
-	EvaluateImageCenterParameters(InFocus, InZoom, InterpolatedImageCenter);
-
-	FFocalLengthInfo InterpolatedFocalLength;
-	EvaluateFocalLength(InFocus, InZoom, InterpolatedFocalLength);
+	using FBlendParams = LensInterpolationUtils::FDistortionMapBlendParams<FDistortionTable>;
+	using FBlendResults = LensInterpolationUtils::FDistortionMapBlendResults;
 	
 	FCameraFilmbackSettings CameraFilmback;
 	CameraFilmback.SensorWidth = InFilmback.X;
@@ -452,64 +427,61 @@ bool ULensFile::EvaluateDistortionForParameters(float InFocus, float InZoom, FVe
 
 	const FVector2D FxFyScale = FVector2D(LensInfo.SensorDimensions.X / CameraFilmback.SensorWidth, LensInfo.SensorDimensions.Y / CameraFilmback.SensorHeight);
 	
+	// Compute interpolated image center and focal length to pass to the lens handler
 	FLensDistortionState InterpolatedState;
-	InterpolatedState.FocalLengthInfo.FxFy = InterpolatedFocalLength.FxFy * FxFyScale;
-	InterpolatedState.ImageCenter.PrincipalPoint = InterpolatedImageCenter.PrincipalPoint;
-		
-	// Initialize all distortion parameters with their default values
-	LensInfo.LensModel->GetDefaultObject<ULensModel>()->GetDefaultParameterArray(InterpolatedState.DistortionInfo.Parameters);
+	EvaluateImageCenterParameters(InFocus, InZoom, InterpolatedState.ImageCenter);
+	EvaluateFocalLength(InFocus, InZoom, InterpolatedState.FocalLengthInfo);
+	InterpolatedState.FocalLengthInfo.FxFy *= FxFyScale;
 	
-	LensInterpolationUtils::FDistortionMapBlendParams<FDistortionTable> Params;
+	FBlendParams Params;
 	Params.bGenerateBlendingParams = true;
+	Params.bGenerateBlendedDistortionParams = true;
+	Params.bGenerateDistortionMaps = true;
+	Params.bGenerateBlendedOverscan = true;
 	Params.UndistortedMaps = UndistortionDisplacementMapHolders;
 	Params.DistortedMaps = DistortionDisplacementMapHolders;
-	Params.DistortionParamNum = InterpolatedState.DistortionInfo.Parameters.Num();
-
-	// Callback that retrieves the distortion parameters for the specified focus and zoom
-	Params.GetDistortionParameters = LensInterpolationUtils::FDistortionMapBlendParams<FDistortionTable>::FGetDistortionParameters::CreateLambda(
-		[](const FDistortionFocusPoint& FocusPoint, const FDistortionFocusCurve& FocusCurve)
+	Params.GetInterpolatedImageCenter = FBlendParams::FGetInterpolatedImageCenter::CreateUObject(this, &ULensFile::EvaluateImageCenterParameters);
+	Params.GetInterpolatedFocalLength = FBlendParams::FGetInterpolatedFocalLength::CreateUObject(this, &ULensFile::EvaluateFocalLength);
+	
+	LensInfo.LensModel->GetDefaultObject<ULensModel>()->GetDefaultParameterArray(Params.DefaultDistortionParams);
+	
+	Params.ProcessDisplacementMaps = FBlendParams::FProcessDisplacementMaps::CreateLambda(
+		[InLensHandler, &CameraFilmback, &FxFyScale](const FLensDistortionState& InDistortionState, UTextureRenderTarget2D* InUndistortionMap, UTextureRenderTarget2D* InDistortionMap)
 		{
-			FDistortionInfo Point;
-			if (FocusPoint.GetPoint(FocusCurve.Zoom, Point))
-			{
-				return TOptional(Point);
-			}
+			FLensDistortionState FinalDistortionState(InDistortionState);
+			FinalDistortionState.FocalLengthInfo.FxFy *= FxFyScale;
+
+			InLensHandler->SetDistortionState(FinalDistortionState);
+			InLensHandler->SetCameraFilmback(CameraFilmback);
 			
-			return TOptional<FDistortionInfo>();
+			InLensHandler->DrawUndistortionDisplacementMap(InUndistortionMap);
+			InLensHandler->DrawDistortionDisplacementMap(InDistortionMap);
 		}
 	);
 
-	// Callback when the blend function constructs the displacement maps for each corner used in the blend, which generates the displacement maps
-	// and returns the computed overscan of the displacement map
-	Params.ProcessDisplacementMaps = LensInterpolationUtils::FDistortionMapBlendParams<FDistortionTable>::FProcessDisplacementMaps::CreateLambda(
-		[this, InLensHandler, &CameraFilmback, &InterpolatedImageCenter, &FxFyScale]
-		(const FDistortionFocusPoint& FocusPoint, const FDistortionFocusCurve& FocusCurve, UTextureRenderTarget2D* InUndistortedMap, UTextureRenderTarget2D* InDistortedMap)
+	Params.GetOverscan = FBlendParams::FGetOverscan::CreateLambda(
+		[InLensHandler, &CameraFilmback, &FxFyScale]
+		(const FDistortionFocusPoint& FocusPoint, float Zoom, const FLensDistortionState& InDistortionState)
 		{
-			FLensDistortionState State;
-			State.ImageCenter.PrincipalPoint = InterpolatedImageCenter.PrincipalPoint;
-
-			// In case the point doesn't exist, we need to fill the distortion parameter array with default values
-			LensInfo.LensModel->GetDefaultObject<ULensModel>()->GetDefaultParameterArray(State.DistortionInfo.Parameters);
-			FocusPoint.GetPoint(FocusCurve.Zoom, State.DistortionInfo);
-
-			FFocalLengthInfo FocalLength;
-			LensDataTableUtils::GetPointValue<FFocalLengthFocusPoint>(FocusPoint.Focus, FocusCurve.Zoom, FocalLengthTable.FocusPoints, FocalLength);
-			
-			State.FocalLengthInfo.FxFy = FocalLength.FxFy * FxFyScale;
+			FLensDistortionState FinalDistortionState(InDistortionState);
+			FinalDistortionState.FocalLengthInfo.FxFy *= FxFyScale;
 
-			InLensHandler->SetDistortionState(State);
+			InLensHandler->SetDistortionState(FinalDistortionState);
 			InLensHandler->SetCameraFilmback(CameraFilmback);
-			InLensHandler->DrawUndistortionDisplacementMap(InUndistortedMap);
-			InLensHandler->DrawDistortionDisplacementMap(InDistortedMap);
 
 			return InLensHandler->ComputeOverscanFactor();
 		}
 	);
 	
-	LensInterpolationUtils::FDistortionMapBlendResults Results = LensInterpolationUtils::DistortionMapBlend(DistortionTable, InFocus, InZoom, Params);
-	if (!Results.bValid)
+	FBlendResults Results = LensInterpolationUtils::DistortionMapBlend(DistortionTable, InFocus, InZoom, Params);
+	if (!Results.bValid || !Results.BlendedDistortionParams.IsSet())
 	{
-		// No distortion parameters case. Still process to have center shift
+		// We were unable to generate blended distortion parameters for some reason. Fallback to default distortion,
+		// while keeping blended image center and focal length
+		
+		// Initialize all distortion parameters with their default values
+		LensInfo.LensModel->GetDefaultObject<ULensModel>()->GetDefaultParameterArray(InterpolatedState.DistortionInfo.Parameters);
+		
 		// Setup handler state based on evaluated parameters. If none were found, no distortion will be returned
 		InLensHandler->SetDistortionState(InterpolatedState);
 		InLensHandler->SetCameraFilmback(CameraFilmback);
@@ -543,7 +515,7 @@ bool ULensFile::EvaluateDistortionForParameters(float InFocus, float InZoom, FVe
 		, DistortionDisplacementMapHolders[2]
 		, DistortionDisplacementMapHolders[3]);
 
-	InLensHandler->SetOverscanFactor(Results.BlendedOverscan.GetValue());
+	InLensHandler->SetOverscanFactor(Results.BlendedOverscan.Get(1.0f));
 	
 	return true;
 }
@@ -552,6 +524,9 @@ bool ULensFile::EvaluateDistortionForSTMaps(float InFocus, float InZoom, FVector
 {
 	TRACE_CPUPROFILER_EVENT_SCOPE(ULensFile::EvaluateDistortionForSTMaps);
 
+	using FBlendParams = LensInterpolationUtils::FDistortionMapBlendParams<FSTMapTable>;
+	using FBlendResults = LensInterpolationUtils::FDistortionMapBlendResults;
+	
 	if(DerivedDataInFlightCount > 0)
 	{
 		UE_LOG(LogCameraCalibrationCore, Verbose, TEXT("Can't evaluate LensFile '%s' - %d data points still being computed. Clearing render target for no distortion"), *GetName(), DerivedDataInFlightCount);
@@ -576,42 +551,30 @@ bool ULensFile::EvaluateDistortionForSTMaps(float InFocus, float InZoom, FVector
 	const FVector2D FxFyScale = FVector2D(InFilmback.X / LensInfo.SensorDimensions.X, InFilmback.Y / LensInfo.SensorDimensions.Y);
 	
 	//When dealing with STMaps, FxFy was not a calibrated value. We can interpolate our curve directly for desired point
-	FFocalLengthInfo FocalLength;
-	EvaluateFocalLength(InFocus, InZoom, FocalLength);
-
-	FImageCenterInfo ImageCenter;
-	EvaluateImageCenterParameters(InFocus, InZoom, ImageCenter);
+	FLensDistortionState InterpolatedState;
+	EvaluateImageCenterParameters(InFocus, InZoom, InterpolatedState.ImageCenter);
+	EvaluateFocalLength(InFocus, InZoom, InterpolatedState.FocalLengthInfo);
+	InterpolatedState.FocalLengthInfo.FxFy *= FxFyScale;
 
-	LensInterpolationUtils::FDistortionMapBlendParams<FSTMapTable> Params;
+	FBlendParams Params;
 	Params.bGenerateBlendingParams = true;
+	Params.bGenerateDistortionMaps = true;
+	Params.bGenerateBlendedOverscan = true;
 
-	// Callback that retrieves the displacement map render targets for the specified focus and zoom
-	Params.GetDisplacementMaps = LensInterpolationUtils::FDistortionMapBlendParams<FSTMapTable>::FGetDisplacementMaps::CreateLambda(
-		[](const FSTMapFocusPoint& FocusPoint, const FSTMapFocusCurve& FocusCurve, UTextureRenderTarget2D*& OutUndistortedMap, UTextureRenderTarget2D*& OutDistortedMap)
+	Params.GetOverscan = FBlendParams::FGetOverscan::CreateLambda(
+		[this]
+		(const FSTMapFocusPoint& FocusPoint, float Zoom, const FLensDistortionState& InDistortionState)
 		{
-			if (const FSTMapZoomPoint* ZoomPoint = FocusPoint.GetZoomPoint(FocusCurve.Zoom))
+			if (const FSTMapZoomPoint* ZoomPoint = FocusPoint.GetZoomPoint(Zoom))
 			{
-				OutUndistortedMap = ZoomPoint->DerivedDistortionData.UndistortionDisplacementMap;
-				OutDistortedMap = ZoomPoint->DerivedDistortionData.DistortionDisplacementMap;
-			}
-		});
-
-	// Callback when the blend function constructs the displacement maps for each corner used in the blend, which generates the displacement maps
-	// and returns the computed overscan of the displacement map
-	Params.ProcessDisplacementMaps = LensInterpolationUtils::FDistortionMapBlendParams<FSTMapTable>::FProcessDisplacementMaps::CreateLambda(
-		[this, &ImageCenter]
-		(const FSTMapFocusPoint& FocusPoint, const FSTMapFocusCurve& FocusCurve, UTextureRenderTarget2D*, UTextureRenderTarget2D*)
-		{
-			if (const FSTMapZoomPoint* ZoomPoint = FocusPoint.GetZoomPoint(FocusCurve.Zoom))
-			{
-				return ComputeOverscan(ZoomPoint->DerivedDistortionData.DistortionData, ImageCenter.PrincipalPoint);
+				return ComputeOverscan(ZoomPoint->DerivedDistortionData.DistortionData, InDistortionState.ImageCenter.PrincipalPoint);
 			}
 
 			return 1.0f;
 		});
 	
-	LensInterpolationUtils::FDistortionMapBlendResults Results = LensInterpolationUtils::DistortionMapBlend(STMapTable, InFocus, InZoom, Params);
-	if (!Results.bValid)
+	FBlendResults Results = LensInterpolationUtils::DistortionMapBlend(STMapTable, InFocus, InZoom, Params);
+	if (!Results.bValid || !Results.UndistortedMaps.IsSet() || !Results.DistortedMaps.IsSet())
 	{
 		UE_LOG(LogCameraCalibrationCore, Verbose, TEXT("Can't evaluate LensFile '%s' - No calibrated maps"), *GetName());
 		SetupNoDistortionOutput(InLensHandler);
@@ -620,7 +583,7 @@ bool ULensFile::EvaluateDistortionForSTMaps(float InFocus, float InZoom, FVector
 
 	FDisplacementMapBlendingParams& BlendingParams = Results.BlendingParams.GetValue();
 	BlendingParams.FxFyScale = FxFyScale;
-	BlendingParams.PrincipalPoint = ImageCenter.PrincipalPoint;
+	BlendingParams.PrincipalPoint = InterpolatedState.ImageCenter.PrincipalPoint;
 	
 	//Draw resulting undistortion displacement map for evaluation point
 	LensFileRendering::DrawBlendedDisplacementMap(InLensHandler->GetUndistortionDisplacementMap()
@@ -638,15 +601,11 @@ bool ULensFile::EvaluateDistortionForSTMaps(float InFocus, float InZoom, FVector
 		, Results.DistortedMaps.GetValue()[2]
 		, Results.DistortedMaps.GetValue()[3]);
 	
-	FLensDistortionState State;
-	State.FocalLengthInfo.FxFy = FocalLength.FxFy * FxFyScale;
-	State.ImageCenter = ImageCenter;
-	
 	//Sets final blended distortion state
-	InLensHandler->SetDistortionState(MoveTemp(State));
+	InLensHandler->SetDistortionState(MoveTemp(InterpolatedState));
 	InLensHandler->SetCameraFilmback(CameraFilmback);
 
-	InLensHandler->SetOverscanFactor(Results.BlendedOverscan.GetValue());
+	InLensHandler->SetOverscanFactor(Results.BlendedOverscan.Get(1.0f));
 	
 	return true;
 }
diff --git a/Engine/Plugins/VirtualProduction/CameraCalibrationCore/Source/CameraCalibrationCore/Private/LensInterpolationUtils.cpp b/Engine/Plugins/VirtualProduction/CameraCalibrationCore/Source/CameraCalibrationCore/Private/LensInterpolationUtils.cpp
--- a/Engine/Plugins/VirtualProduction/CameraCalibrationCore/Source/CameraCalibrationCore/Private/LensInterpolationUtils.cpp
+++ b/Engine/Plugins/VirtualProduction/CameraCalibrationCore/Source/CameraCalibrationCore/Private/LensInterpolationUtils.cpp
@@ -168,6 +168,47 @@ namespace LensInterpolationUtils
 		return UE::Curves::BezierInterp(P0, P1, P2, P3, Alpha);
 	}
 
+	float FTangentBezierCurve::EvalTangent(float InX)
+	{
+		if (FMath::IsNearlyEqual(X0, X1))
+		{
+			return Tangent0;
+		}
+
+		if (FMath::IsNearlyEqual(Y0, Y1))
+		{
+			return 0.0f;
+		}
+
+		const float DeltaX = X1 - X0;
+		const float DeltaY = Y1 - Y0;
+		const float Alpha = (InX - X0) / DeltaX;
+
+		const float T0 = Tangent0 * DeltaY;
+		const float T2 = (3.0f - (Tangent1 + Tangent0)) * DeltaY;
+		const float T3 = Tangent1 * DeltaY;
+		
+		return FMath::Lerp(FMath::Lerp(T0, T2, Alpha), FMath::Lerp(T2, T3, Alpha), Alpha) / DeltaY;
+	}
+	
+	TArray<float> FTangentBezierCurve::MultiEval(float InX, const TArray<float>& InY0s, const TArray<float>& InY1s)
+	{
+		check(InY0s.Num() == InY1s.Num());
+		
+		TArray<float> Results;
+		Results.AddZeroed(InY0s.Num());
+
+		for (int32 Index = 0; Index < InY0s.Num(); ++Index)
+		{
+			Y0 = InY0s[Index];
+			Y1 = InY1s[Index];
+
+			Results[Index] = Eval(InX);
+		}
+
+		return Results;
+	}
+
 	float FCoonsPatch::Blend(const FVector2D& InPoint)
 	{
 		float Alpha = 0.0;
@@ -197,39 +238,184 @@ namespace LensInterpolationUtils
 		return Lx + Ly - 0.5 * B;
 	}
 
-	FTangentBezierCoonsPatch::FTangentBezierCoonsPatch(float X0, float X1, float Y0, float Y1, float XTangents[4], float YTangents[4])
-		: Corners {
-			FPatchCorner(X0, Y0, XTangents[0], YTangents[0], 0.0),
-			FPatchCorner(X1, Y0, XTangents[1], YTangents[1], 0.0),
-			FPatchCorner(X1, Y1, XTangents[2], YTangents[2], 0.0),
-			FPatchCorner(X0, Y1, XTangents[3], YTangents[3], 0.0)}
-	{ }
-
 	float FTangentBezierCoonsPatch::Blend(const FVector2D& InPoint)
 	{
-		
-		float Alpha = 0.0;
-		if (!FMath::IsNearlyEqual(Corners[0].X, Corners[1].X))
+		float Alpha = 0.0f;
+		if (!FMath::IsNearlyEqual(YCurves[0].Coordinate, YCurves[1].Coordinate))
 		{
-			Alpha = (InPoint.X - Corners[0].X) / (Corners[1].X - Corners[0].X);
+			Alpha = (InPoint.X - YCurves[0].Coordinate) / (YCurves[1].Coordinate - YCurves[0].Coordinate);
 		}
+
+		float Beta = 0.0f;
+		if (!FMath::IsNearlyEqual(XCurves[0].Coordinate, XCurves[1].Coordinate))
+		{
+			Beta = (InPoint.Y - XCurves[0].Coordinate) / (XCurves[1].Coordinate - XCurves[0].Coordinate);
+		}
+
+		const float Lx = FMath::Lerp(XCurves[0].TangentCurve.Eval(InPoint.X), XCurves[0].TangentCurve.Eval(InPoint.X), Beta);
+		const float Ly = FMath::Lerp(YCurves[0].TangentCurve.Eval(InPoint.Y), YCurves[0].TangentCurve.Eval(InPoint.Y), Alpha);
+
+		float PatchCorners[4];
+		for (int32 Index = 0; Index < 4; ++Index)
+		{
+			const int32 XIndex = Index / 2;
+			const int32 YIndex = FMath::WrapExclusive(Index + 1, 0, 4) / 2;
+
+			PatchCorners[Index] = 0.5f * (XCurves[XIndex].TangentCurve.Eval(YCurves[YIndex].Coordinate) + YCurves[YIndex].TangentCurve.Eval(XCurves[XIndex].Coordinate));
+		}
+
+		const float B = FMath::BiLerp(PatchCorners[0], PatchCorners[1], PatchCorners[3], PatchCorners[2], Alpha, Beta);
+		return Lx + Ly - B;
+	}
+
+	FTangentBezierCurve GetTangentCurve(const FRichCurve& Curve, float Time, float MinTime, float MaxTime)
+	{
+		FTangentBezierCurve OutCurve(MinTime, MaxTime, 0.0, 0.0, 0.0, 0.0);
 		
-		float Beta = 0.0;
-		if (!FMath::IsNearlyEqual(Corners[0].Y, Corners[3].Y))
+		const int32 NumKeys = Curve.GetNumKeys();
+		if (NumKeys == 0)
 		{
-			Beta = (InPoint.Y - Corners[0].Y) / (Corners[3].Y - Corners[0].Y);
+			return OutCurve;
 		}
 
-		FTangentBezierCurve X0Curve = FTangentBezierCurve(Corners[0].X, Corners[1].X, Corners[0].Value, Corners[1].Value, Corners[0].TangentX, Corners[1].TangentX);
-		FTangentBezierCurve X1Curve = FTangentBezierCurve(Corners[3].X, Corners[2].X, Corners[3].Value, Corners[2].Value, Corners[3].TangentX, Corners[2].TangentX);
-		FTangentBezierCurve Y0Curve = FTangentBezierCurve(Corners[0].Y, Corners[3].Y, Corners[0].Value, Corners[3].Value, Corners[0].TangentY, Corners[3].TangentY);
-		FTangentBezierCurve Y1Curve = FTangentBezierCurve(Corners[1].Y, Corners[2].Y, Corners[1].Value, Corners[2].Value, Corners[1].TangentY, Corners[2].TangentY);
-	
-		const float Lx = FMath::Lerp(X0Curve.Eval(InPoint.X), X1Curve.Eval(InPoint.X), Beta);
-		const float Ly = FMath::Lerp(Y0Curve.Eval(InPoint.Y), Y1Curve.Eval(InPoint.Y), Alpha);
-		const float B = FMath::BiLerp(Corners[0].Value, Corners[1].Value, Corners[3].Value, Corners[2].Value, Alpha, Beta);
+		if (Time < Curve.Keys[0].Time)
+		{
+			OutCurve.X1 = Curve.Keys[0].Time;
+			return OutCurve;
+		}
+		if (Time >= Curve.Keys[NumKeys - 1].Time)
+		{
+			OutCurve.X0 = Curve.Keys[NumKeys - 1].Time;
+			return OutCurve;
+		}
+
+		// Similar to RichCurve.Eval(), perform a binary search on the curve's keys to find the bounding points around Time
+		int32 Start = 0;
+		int32 End = NumKeys - 1;
+
+		while (Start < End)
+		{
+			const int32 TestPos = Start + (End-Start) / 2;
+			if (Time >= Curve.Keys[TestPos].Time)
+			{
+				Start = TestPos + 1;
+			}
+			else
+			{
+				End = TestPos;
+			}
+		}
+
+		const FRichCurveKey& Prev = Curve.Keys[Start - 1];
+		const FRichCurveKey& Next = Curve.Keys[Start];
+
+		OutCurve.X0 = Prev.Time;
+		OutCurve.X1 = Next.Time;
 		
-		return Lx + Ly - B;
+		if (Prev.InterpMode == ERichCurveInterpMode::RCIM_Cubic)
+		{
+			OutCurve.Tangent0 = Prev.LeaveTangent;
+			OutCurve.Tangent1 = Next.ArriveTangent;
+		}
+		else
+		{
+			float Slope = (Next.Value - Prev.Value) / (Next.Time - Prev.Time);
+			OutCurve.Tangent0 = Slope;
+			OutCurve.Tangent1 = Slope;
+		}
+
+		return OutCurve;
+	}
+
+	template <>
+	TOptional<FDistortionInfo> GetDistortionInfo<FDistortionTable>(const FDistortionFocusPoint& FocusPoint, float Zoom)
+	{
+		FDistortionInfo DistortionInfo;
+		if (FocusPoint.GetPoint(Zoom, DistortionInfo))
+		{
+			return DistortionInfo;
+		}
+
+		return TOptional<FDistortionInfo>();
+	}
+
+	template <>
+	void GetDisplacementMaps<FSTMapTable>(const FSTMapFocusPoint& FocusPoint, float Zoom, UTextureRenderTarget2D*& OutUndistortedMap, UTextureRenderTarget2D*& OutDistortedMap)
+	{
+		if (const FSTMapZoomPoint* ZoomPoint = FocusPoint.GetZoomPoint(Zoom))
+		{
+			OutUndistortedMap = ZoomPoint->DerivedDistortionData.UndistortionDisplacementMap;
+			OutDistortedMap = ZoomPoint->DerivedDistortionData.DistortionDisplacementMap;
+		}
+	}
+
+	template<>
+	void FDistortionMapBlendParams<FDistortionTable>::InitializeResults(float InFocus, float InZoom, FDistortionMapBlendResults& OutResults) const
+	{
+		if (bGenerateBlendingParams)
+		{
+			OutResults.bValid = true;
+
+			OutResults.BlendingParams = FDisplacementMapBlendingParams();
+			OutResults.BlendingParams->EvalFocus = InFocus;
+			OutResults.BlendingParams->EvalZoom = InZoom;
+		}
+
+		if (bGenerateBlendedDistortionParams)
+		{
+			OutResults.bValid = true;
+			
+			OutResults.BlendedDistortionParams = FDistortionInfo();
+			OutResults.BlendedDistortionParams->Parameters = DefaultDistortionParams;
+		}
+
+		if (bGenerateDistortionMaps)
+		{
+			if (ProcessDisplacementMaps.IsBound() && UndistortedMaps.Num() == 4 && DistortedMaps.Num() == 4)
+			{
+				OutResults.bValid = true;
+				
+				OutResults.UndistortedMaps = UndistortedMaps;
+				OutResults.DistortedMaps = DistortedMaps;
+			}
+		}
+
+		if (bGenerateBlendedOverscan && GetOverscan.IsBound())
+		{
+			OutResults.bValid = true;
+			
+			OutResults.BlendedOverscan = 1.0f;
+		}
+	}
+
+	template<>
+	void FDistortionMapBlendParams<FSTMapTable>::InitializeResults(float InFocus, float InZoom, FDistortionMapBlendResults& OutResults) const
+	{
+		if (bGenerateBlendingParams)
+		{
+			OutResults.bValid = true;
+			
+			OutResults.BlendingParams = FDisplacementMapBlendingParams();
+			OutResults.BlendingParams->EvalFocus = InFocus;
+			OutResults.BlendingParams->EvalZoom = InZoom;
+		}
+
+		if (bGenerateDistortionMaps)
+		{
+			OutResults.bValid = true;
+			
+			OutResults.UndistortedMaps = TArray<UTextureRenderTarget2D*>();
+			OutResults.DistortedMaps = TArray<UTextureRenderTarget2D*>();
+			OutResults.UndistortedMaps->AddZeroed(4);
+			OutResults.DistortedMaps->AddZeroed(4);
+		}
+
+		if (bGenerateBlendedOverscan && GetOverscan.IsBound())
+		{
+			OutResults.bValid = true;
+			
+			OutResults.BlendedOverscan = 1.0f;
+		}
 	}
 }
 
diff --git a/Engine/Plugins/VirtualProduction/CameraCalibrationCore/Source/CameraCalibrationCore/Private/LensInterpolationUtils.h b/Engine/Plugins/VirtualProduction/CameraCalibrationCore/Source/CameraCalibrationCore/Private/LensInterpolationUtils.h
--- a/Engine/Plugins/VirtualProduction/CameraCalibrationCore/Source/CameraCalibrationCore/Private/LensInterpolationUtils.h
+++ b/Engine/Plugins/VirtualProduction/CameraCalibrationCore/Source/CameraCalibrationCore/Private/LensInterpolationUtils.h
@@ -7,11 +7,9 @@
 #include "LensDistortionSceneViewExtension.h"
 #include "LensFileRendering.h"
 #include "Math/Vector2D.h"
-#include "Misc/EnumClassFlags.h"
 #include "Tables/DistortionParametersTable.h"
 #include "Tables/LensTableUtils.h"
-#include "UObject/Object.h"
-#include "UObject/ObjectMacros.h"
+#include "Tables/STMapTable.h"
 
 
 struct FRichCurve;
@@ -50,6 +48,12 @@ namespace LensInterpolationUtils
 		/** Evaluates the Bezier curve at the specified x value */
 		float Eval(float InX);
 
+		/** Evaluates the tangent of the Bezier curve at the specified x value */
+		float EvalTangent(float InX);
+		
+		/** Evaluates the Bezier curve for each of the provided y values at the specified x value */
+		TArray<float> MultiEval(float InX, const TArray<float>& InY0s, const TArray<float>& InY1s);
+		
 		/** Minimum x value of the curve */
 		float X0;
 
@@ -116,37 +120,51 @@ namespace LensInterpolationUtils
 	};
 
 	/**
-	 * A version of the Coons patch that takes in four points that each have x and y tangent data, and performs
-	 * a Bezier interpolation to create the four edge curves from the points and their tangents
+	 * A version of the Coons patch that takes in four tangent Bezier curves, and performs
+	 * a Bezier interpolation to create the four edge curves from the points and their tangents.
+	 * Works in cases where the curves' endpoints do not match, in which case the four patch corners are
+	 * computed by interpolation and averaged before the final Coons patch interpolation
 	 */
 	struct FTangentBezierCoonsPatch
 	{
-		/** Creates a new Coons patch from four corner points with x and y tangents. */
-		FTangentBezierCoonsPatch(float X0, float X1, float Y0, float Y1, float XTangents[4], float YTangents[4]);
+		struct FPatchCurve
+		{
+			/** The tangent curve along the primary axis of the patch edge */
+			FTangentBezierCurve TangentCurve;
+
+			/** The coordinate along the secondary axis that the curve is at */
+			float Coordinate;
+
+			FPatchCurve(FTangentBezierCurve InTangentCurve, float InCoordinate)
+				: TangentCurve(InTangentCurve)
+				, Coordinate(InCoordinate)
+			{}
+		};
+		
+		FTangentBezierCoonsPatch(const FPatchCurve& X0Curve, const FPatchCurve& X1Curve, const FPatchCurve& Y0Curve, const FPatchCurve& Y1Curve)
+			: XCurves { X0Curve, X1Curve }
+			, YCurves { Y0Curve, Y1Curve }
+		{}
 		
 		/** Computes the value of the patch at the specified point by blending between the four edge curves */
 		float Blend(const FVector2D& InPoint);
 
-		/** Stores the x and y position and x and y tangents of a corner */
-		struct FPatchCorner
-		{
-			float X;
-			float Y;
-			float TangentX;
-			float TangentY;
-			float Value;
-
-			FPatchCorner(float InX, float InY, float InTangentX, float InTangentY, float InValue)
-				: X(InX), Y(InY), TangentX(InTangentX), TangentY(InTangentY), Value(InValue)
-			{}
-		};
+		/** The tangent curves along the x axes edges of the patch */
+		FPatchCurve XCurves[2];
 
-		/**
-		 * The four corners of the patch, each of which contains and x and y coordinate, and x and y tangent, and a value.
-		 * Indexed in the following order: (X0, Y0) -> (X1, Y0) -> (X1, Y1) -> (X0, Y1)
-		 */
-		FPatchCorner Corners[4];
+		/** The tangent curves along the y axes edges of the patch */
+		FPatchCurve YCurves[2];
 	};
+
+	/**
+	 * Constructs an appropriate tangent curve by seeking the nearest defined points that the specified Time is between
+	 * @param Curve The curve to find defined points from
+	 * @param Time The time value to find bounding points for
+	 * @param MinTime If no defined point can be found less than Time, use this time value as the minimum
+	 * @param MaxTime If no defined point can be found greater than Time, use this time value as the maximum
+	 * @return A tangent Bezier curve that can be interpolated on
+	 */
+	FTangentBezierCurve GetTangentCurve(const FRichCurve& Curve, float Time, float MinTime, float MaxTime);
 	
 	/** Performs a Coons patch blend on an indexed list of parameters where each parameter has its own set of curves, and outputs all blended parameters */
 	template<typename FocusPointType, typename FocusCurveType>
@@ -217,74 +235,46 @@ namespace LensInterpolationUtils
 		return true;
 	}
 
-	/**
-	 * Parameters for a distortion map blend that configure what values are computed for the blend and how to retrieve the necessary
-	 * blending values from each point being blended
-	 */
+	/** Gets the distortion info for the specified focus point and zoom. By default, returns an unset TOptional */
 	template<typename TableType>
-	struct FDistortionMapBlendParams
-	{
-		using FocusPointType = typename TableType::FocusPointType;
-		using FocusCurveType = typename TableType::FocusCurveType;
+	TOptional<FDistortionInfo> GetDistortionInfo(const typename TableType::FocusPointType& FocusPoint, float Zoom) { return TOptional<FDistortionInfo>(); }
 
-		DECLARE_DELEGATE_FourParams(FGetDisplacementMaps,
-			const FocusPointType& /* FocusPoint */,
-			const FocusCurveType& /* FocusCurve */,
-			UTextureRenderTarget2D*& /* OutUndistortedMap */,
-			UTextureRenderTarget2D*& /* OutDistortedMap */);
+	/** Template specialization of GetDistortionInfo for FDistortionTable */
+	template<>
+	TOptional<FDistortionInfo> GetDistortionInfo<FDistortionTable>(const FDistortionFocusPoint& FocusPoint, float Zoom);
 		
-		DECLARE_DELEGATE_RetVal_TwoParams(TOptional<FDistortionInfo>, FGetDistortionParameters,
-			const FocusPointType& /* FocusPoint */,
-			const FocusCurveType& /* FocusCurve */);
-		
-		DECLARE_DELEGATE_RetVal_FourParams(float, FProcessDisplacementMaps,
-			const FocusPointType& /* FocusPoint */,
-			const FocusCurveType& /* FocusCurve */,
-			UTextureRenderTarget2D* /* InUndistortedMap */,
-			UTextureRenderTarget2D* /* InDistortedMap */);
-
-		DECLARE_DELEGATE_ThreeParams(FGetDistortionState,
-			const FocusPointType& /* FocusPoint */,
-			const FocusCurveType& /* FocusCurve */,
-			FLensDistortionState& /* OutState */);
-
-		/**
-		 * Callback to use to retrieve the distortion parameters for a specified point being blended.
-		 * If unbound, no distortion parameters will be computed in the blending results
-		 */
-		FGetDistortionParameters GetDistortionParameters;
-
-		/**
-		 * Callback to retrieve the displacement map render targets to use for a specified point being blended.
-		 * If unbound, use the provided displacement maps in the parameters
-		 */
-		FGetDisplacementMaps GetDisplacementMaps;
+	/** Gets the pointers to displacement maps for the specified focus point and zoom. Does nothing by default */
+	template<typename TableType>
+	void GetDisplacementMaps(const typename TableType::FocusPointType& FocusPoint, float Zoom, UTextureRenderTarget2D*& OutUndistortedMap, UTextureRenderTarget2D*& OutDistortedMap) { }
 		
-		/**
-		 * Callback to use to perform any processing on the distortion maps for blending that returns the computed overscan for the maps.
-		 * If unbound, no overscan blending will be computed in the blending results
-		 */
-		FProcessDisplacementMaps ProcessDisplacementMaps;
+	/** Template specialization of GetDisplacementMaps for FSTMapTable */
+	template<>
+	void GetDisplacementMaps<FSTMapTable>(const FSTMapFocusPoint& FocusPoint, float Zoom, UTextureRenderTarget2D*& OutUndistortedMap, UTextureRenderTarget2D*& OutDistortedMap);
 
-		/**
-		 * Callback to use to retrieve the complete distortion state for a specified point being blended.
-		 * If unbound, no distortion state will be added to the blending results
-		 */
-		FGetDistortionState GetDistortionState;
-
-		/** Indicates that shader blending parameters should be calculated */
-		bool bGenerateBlendingParams = false;
+	/** Finds the focus points for the specified previous and next focus, with special handling for the case where there is no previous or next focus point */
+	template<typename TableType>
+	LensDataTableUtils::FPointNeighbors GetOnCurveFocusPoints(const TConstArrayView<typename TableType::FocusPointType>& FocusPoints, float PrevFocus, float NextFocus)
+	{
+		LensDataTableUtils::FPointNeighbors Result;
 		
-		/** The number of distortion parameters to blend */
-		int32 DistortionParamNum = 0;
+		int32 PrevIndex = FocusPoints.IndexOfByPredicate([PrevFocus](const typename TableType::FocusPointType& FocusPoint)
+		{
+			return FMath::IsNearlyEqual(FocusPoint.Focus, PrevFocus);
+		});
+		int32 NextIndex = FocusPoints.IndexOfByPredicate([NextFocus](const typename TableType::FocusPointType& FocusPoint)
+		{
+			return FMath::IsNearlyEqual(FocusPoint.Focus, NextFocus);
+		});
 
-		/** When supplied, list of render targets to write the undistorted displacement maps to */
-		TArray<UTextureRenderTarget2D*> UndistortedMaps;
+		if (FocusPoints.IsValidIndex(PrevIndex) || FocusPoints.IsValidIndex(NextIndex))
+		{
+			// If one of the endpoints does not exist, assume the curve is flat, and set both endpoints to the same focus point
+			Result.PreviousIndex = FocusPoints.IsValidIndex(PrevIndex) ? PrevIndex : NextIndex;
+			Result.NextIndex = FocusPoints.IsValidIndex(NextIndex) ? NextIndex : PrevIndex;
+		}
 
-		/** When supplied, list of render targets to write the distorted displacement maps to */
-		TArray<UTextureRenderTarget2D*> DistortedMaps;
-		
-	};
+		return Result;
+	}
 
 	/** Blended results from a distortion map blend. Optionals are set when the blend input parameters indicated those blended values should be computed */
 	struct FDistortionMapBlendResults
@@ -308,6 +298,78 @@ namespace LensInterpolationUtils
 		TOptional<float> BlendedOverscan = TOptional<float>();
 	};
 
+	/**
+	 * Parameters for a distortion map blend that configure what values are computed for the blend and how to retrieve the necessary
+	 * blending values from each point being blended
+	 */
+	template<typename TableType>
+	struct FDistortionMapBlendParams
+	{
+		using FocusPointType = typename TableType::FocusPointType;
+		using FocusCurveType = typename TableType::FocusCurveType;
+		
+		DECLARE_DELEGATE_RetVal_ThreeParams(bool, FGetInterpolatedImageCenter,
+			float /* InFocus */,
+			float /* InZoom */,
+			FImageCenterInfo& /* OutImageCenterInfo */);
+
+		DECLARE_DELEGATE_RetVal_ThreeParams(bool, FGetInterpolatedFocalLength,
+			float /* InFocus */,
+			float /* InZoom */,
+			FFocalLengthInfo& /* OutFocalLengthInfo */);
+
+		DECLARE_DELEGATE_ThreeParams(FProcessDisplacementMaps,
+			const FLensDistortionState& /* InDistortionState */,
+			UTextureRenderTarget2D* /* OutUndistortionMap */,
+			UTextureRenderTarget2D* /* OutDistortionMap */);
+
+		DECLARE_DELEGATE_RetVal_ThreeParams(float, FGetOverscan,
+			const FocusPointType& /* InFocusPoint */,
+			float /* InZoom */,
+			const FLensDistortionState& /* InDistortionState */);
+
+		/** Callback to use when the interpolated image center needs to be computed */
+		FGetInterpolatedImageCenter GetInterpolatedImageCenter;
+		
+		/** Callback to use when the interpolated image center needs to be computed */
+		FGetInterpolatedFocalLength GetInterpolatedFocalLength;
+		
+		FProcessDisplacementMaps ProcessDisplacementMaps;
+
+		/** Callback to compute the overscan for a specific distortion state */
+		FGetOverscan GetOverscan;
+		
+		/** Indicates that shader blending parameters (Results.BlendingParams) should be calculated */
+		bool bGenerateBlendingParams = false;
+
+		/** Indicates whether distortion parameters (Results.BlendedDistortionParams) should be interpolated at the focus and zoom */
+		bool bGenerateBlendedDistortionParams = false;
+
+		/** Indicates whether distortion maps (Results.UndistortedMaps and Results.DistortedMaps) should be generated */
+		bool bGenerateDistortionMaps = false;
+
+		/** Indicates whether overscan (Results.BlendedOverscan) should be generated */ 
+		bool bGenerateBlendedOverscan = false;
+
+		/** The default distortion parameters to use for cases when there is no distortion data in the table being interpolated over */
+		TArray<float> DefaultDistortionParams;
+		
+		/** When supplied, list of render targets to write the undistorted displacement maps to */
+		TArray<UTextureRenderTarget2D*> UndistortedMaps;
+
+		/** When supplied, list of render targets to write the distorted displacement maps to */
+		TArray<UTextureRenderTarget2D*> DistortedMaps;
+
+		/** Initializes the blend results as needed based on the state of the parameters. OutResults.bValid will be false if an issue was found with the parameters */
+		void InitializeResults(float InFocus, float InZoom, FDistortionMapBlendResults& OutResults) const { }
+	};
+
+	template<>
+	void FDistortionMapBlendParams<FDistortionTable>::InitializeResults(float InFocus, float InZoom, FDistortionMapBlendResults& OutResults) const;
+
+	template<>
+	void FDistortionMapBlendParams<FSTMapTable>::InitializeResults(float InFocus, float InZoom, FDistortionMapBlendResults& OutResults) const;
+	
 	/** Performs a Coons patch blend on the distortion map parameters for the specified table at the specified focus and zoom */
 	template<typename TableType>
 	FDistortionMapBlendResults DistortionMapBlend(const TableType& Table, float InFocus, float InZoom, const FDistortionMapBlendParams<TableType>& InParams)
@@ -316,37 +378,6 @@ namespace LensInterpolationUtils
 		
 		using FocusPointType = typename TableType::FocusPointType;
 		using FocusCurveType = typename TableType::FocusCurveType;
-
-		// Function to get the correct tangents of the curves safely, even in edge cases where there isn't a rich curve key defined for the specified point
-		auto SafeGetTangents = [](const FRichCurve& Curve, const LensDataTableUtils::FPointNeighbors& Points, float& OutPrevTangent, float& OutNextTangent)
-		{
-			bool bHasPrevious = Curve.Keys.IsValidIndex(Points.PreviousIndex);
-			bool bHasNext = Curve.Keys.IsValidIndex(Points.NextIndex);
-		
-			if (bHasPrevious && bHasNext)
-			{
-				const FRichCurveKey& Prev = Curve.Keys[Points.PreviousIndex];
-				const FRichCurveKey& Next = Curve.Keys[Points.NextIndex];
-				if (Prev.InterpMode == ERichCurveInterpMode::RCIM_Cubic)
-				{
-					OutPrevTangent = Prev.LeaveTangent;
-					OutNextTangent = Next.ArriveTangent;
-				}
-				else
-				{
-					float Slope = (Next.Value - Prev.Value) / (Next.Time - Prev.Time);
-					OutPrevTangent = Slope;
-					OutNextTangent = Slope;
-				}
-			}
-			else
-			{
-				// If one or both points don't have a key in the curve, the curve between those two points
-				// will necessarily be flat, and as such both tangents are 0
-				OutPrevTangent = 0;
-				OutNextTangent = 0;
-			}
-		};
 		
 		TConstArrayView<FocusPointType> FocusPoints = Table.GetFocusPoints();
 		TConstArrayView<FocusCurveType> FocusCurves = Table.GetFocusCurves();
@@ -356,38 +387,11 @@ namespace LensInterpolationUtils
 			return Results;
 		}
 
-		Results.bValid = true;
-
-		if (InParams.bGenerateBlendingParams)
-		{
-			Results.BlendingParams = FDisplacementMapBlendingParams();
-			Results.BlendingParams->EvalFocus = InFocus;
-			Results.BlendingParams->EvalZoom = InZoom;
-		}
-
-		if (InParams.ProcessDisplacementMaps.IsBound())
-		{
-			// If the callback to retrieve the displacement maps for each blended point is set, use that as a source for the displacement maps
-			// Otherwise, if any displacement maps were passed in the parameters, use them
-			if (InParams.GetDisplacementMaps.IsBound())
-			{
-				Results.UndistortedMaps = TArray<UTextureRenderTarget2D*>();
-				Results.DistortedMaps = TArray<UTextureRenderTarget2D*>();
-
-				Results.UndistortedMaps->AddZeroed(4);
-				Results.DistortedMaps->AddZeroed(4);
-			}
-			else if (InParams.UndistortedMaps.Num() == 4 && InParams.DistortedMaps.Num() == 4)
-			{
-				Results.UndistortedMaps = InParams.UndistortedMaps;
-				Results.DistortedMaps = InParams.DistortedMaps;
-			}
-		}
-		
-		if (InParams.GetDistortionParameters.IsBound())
+		InParams.InitializeResults(InFocus, InZoom, Results);
+		if (!Results.bValid)
 		{
-			Results.BlendedDistortionParams = FDistortionInfo();
-			Results.BlendedDistortionParams->Parameters.AddDefaulted(InParams.DistortionParamNum);
+			// Something is wrong with the input parameters, and they could not initialize the results, so abort
+			return Results;
 		}
 		
 		const LensDataTableUtils::FPointNeighbors PointNeighbors = LensDataTableUtils::FindFocusPoints(InFocus, FocusPoints);
@@ -397,205 +401,397 @@ namespace LensInterpolationUtils
 		const FocusPointType& NextFocusPoint = FocusPoints[PointNeighbors.NextIndex];
 		const FocusCurveType& PrevFocusCurve = FocusCurves[CurveNeighbors.PreviousIndex];
 		const FocusCurveType& NextFocusCurve = FocusCurves[CurveNeighbors.NextIndex];
-		
-		if (PointNeighbors.IsSinglePoint() && CurveNeighbors.IsSinglePoint())
+
+		struct FCurveEndpoints
+		{
+			FDistortionInfo X0;
+			FDistortionInfo X1;
+		};
+
+		// Helper that can retrieve distortion info for the endpoints of curves, and correctly handles the case where
+		// no data exists for one or both endpoints of the curve.
+		auto CreateCurveEndpoints = [&InParams](TOptional<FDistortionInfo> InX0, TOptional<FDistortionInfo> InX1)
 		{
-			// We are on a corner of the blending patch
-			// Simply retrieve the value at that point
+			if (InX0.IsSet() && InX1.IsSet())
+			{
+				return FCurveEndpoints { InX0.GetValue(), InX1.GetValue() };
+			}
+			
+			if (InX0.IsSet())
+			{
+				return FCurveEndpoints { InX0.GetValue(), InX0.GetValue() };
+			}
+
+			if (InX1.IsSet())
+			{
+				return FCurveEndpoints { InX1.GetValue(), InX1.GetValue() };
+			}
+
+			FCurveEndpoints DefaultEndpoints;
+			
+			DefaultEndpoints.X0.Parameters = InParams.DefaultDistortionParams;
+			DefaultEndpoints.X1.Parameters = InParams.DefaultDistortionParams;
+
+			return DefaultEndpoints;
+		};
+
+		if (PointNeighbors.IsSinglePoint() && CurveNeighbors.IsSinglePoint() && PrevFocusPoint.MapBlendingCurve.FindKey(PrevFocusCurve.Zoom) != FKeyHandle::Invalid())
+		{
+			// We are at the intersection of two curves, with a defined point. Simply return the values at that point
+			
+			FLensDistortionState InterpolationState;
+
+			FDistortionInfo DefaultInfo;
+			DefaultInfo.Parameters = InParams.DefaultDistortionParams;
+			
+			InterpolationState.DistortionInfo = GetDistortionInfo<TableType>(PrevFocusPoint, PrevFocusCurve.Zoom).Get(DefaultInfo);
+
+			if (InParams.GetInterpolatedImageCenter.IsBound())
+			{
+				InParams.GetInterpolatedImageCenter.Execute(PrevFocusPoint.Focus, PrevFocusCurve.Zoom, InterpolationState.ImageCenter);
+			}
+
+			if (InParams.GetInterpolatedFocalLength.IsBound())
+			{
+				InParams.GetInterpolatedFocalLength.Execute(PrevFocusPoint.Focus, PrevFocusCurve.Zoom, InterpolationState.FocalLengthInfo);
+			}
+			
 			if (Results.BlendingParams.IsSet())
 			{
 				Results.BlendingParams->BlendType = EDisplacementMapBlendType::OneFocusOneZoom;
-				InParams.GetDistortionState.ExecuteIfBound(PrevFocusPoint, PrevFocusCurve, Results.BlendingParams->States[0]);
+				Results.BlendingParams->States[0] = InterpolationState;
 			}
 
 			if (Results.BlendedDistortionParams.IsSet())
 			{
-				TOptional<FDistortionInfo> Parameters = InParams.GetDistortionParameters.Execute(PrevFocusPoint, PrevFocusCurve);
-				if (Parameters.IsSet())
+				Results.BlendedDistortionParams = InterpolationState.DistortionInfo;
+			}
+
+			if (Results.UndistortedMaps.IsSet())
+			{
+				GetDisplacementMaps<TableType>(PrevFocusPoint, PrevFocusCurve.Zoom, (*Results.UndistortedMaps)[0], (*Results.DistortedMaps)[0]);
+				
+				if (InParams.ProcessDisplacementMaps.IsBound())
 				{
-					Results.BlendedDistortionParams = Parameters.GetValue();
+					InParams.ProcessDisplacementMaps.Execute(InterpolationState, (*Results.UndistortedMaps)[0], (*Results.DistortedMaps)[0]);
 				}
 			}
 
-			if (Results.UndistortedMaps.IsSet())
+			if (Results.BlendedOverscan.IsSet())
 			{
-				InParams.GetDisplacementMaps.ExecuteIfBound(PrevFocusPoint, PrevFocusCurve, (*Results.UndistortedMaps)[0], (*Results.DistortedMaps)[0]);
-				Results.BlendedOverscan = InParams.ProcessDisplacementMaps.Execute(PrevFocusPoint, PrevFocusCurve, (*Results.UndistortedMaps)[0], (*Results.DistortedMaps)[0]);
+				Results.BlendedOverscan = InParams.GetOverscan.Execute(PrevFocusPoint, PrevFocusCurve.Zoom, InterpolationState);
 			}
 		}
 		else if (PointNeighbors.IsSinglePoint())
 		{
-			// We are on a zoom curve edge of the blending patch
-			float PrevTangent, NextTangent;
-			SafeGetTangents(PrevFocusPoint.MapBlendingCurve, CurveNeighbors, PrevTangent, NextTangent);
+			FTangentBezierCurve BlendCurve = GetTangentCurve(PrevFocusPoint.MapBlendingCurve, InZoom, PrevFocusCurve.Zoom, NextFocusCurve.Zoom);
+			FCurveEndpoints CurveEndpoints = CreateCurveEndpoints(
+				GetDistortionInfo<TableType>(PrevFocusPoint, BlendCurve.X0),
+				GetDistortionInfo<TableType>(PrevFocusPoint, BlendCurve.X1));
 			
-			FTangentBezierCurve BlendCurve(PrevFocusCurve.Zoom, NextFocusCurve.Zoom, 0.0f, 0.0f, PrevTangent, NextTangent);
+			FLensDistortionState InterpolationStates[2];
+			InterpolationStates[0].DistortionInfo = CurveEndpoints.X0;
+			InterpolationStates[1].DistortionInfo = CurveEndpoints.X1;
+
+			if (InParams.GetInterpolatedImageCenter.IsBound())
+			{
+				InParams.GetInterpolatedImageCenter.Execute(PrevFocusPoint.Focus, BlendCurve.X0, InterpolationStates[0].ImageCenter);
+				InParams.GetInterpolatedImageCenter.Execute(PrevFocusPoint.Focus, BlendCurve.X1, InterpolationStates[1].ImageCenter);
+			}
+
+			if (InParams.GetInterpolatedFocalLength.IsBound())
+			{
+				InParams.GetInterpolatedFocalLength.Execute(PrevFocusPoint.Focus, BlendCurve.X0, InterpolationStates[0].FocalLengthInfo);
+				InParams.GetInterpolatedFocalLength.Execute(PrevFocusPoint.Focus, BlendCurve.X1, InterpolationStates[1].FocalLengthInfo);
+			}
 			
 			if (Results.BlendingParams.IsSet())
 			{
 				Results.BlendingParams->BlendType = EDisplacementMapBlendType::OneFocusTwoZoom;
-				Results.BlendingParams->PatchCorners[0] = FDisplacementMapBlendPatchCorner(PrevFocusCurve.Zoom, PrevFocusPoint.Focus, PrevTangent, 0.0);
-				Results.BlendingParams->PatchCorners[1] = FDisplacementMapBlendPatchCorner(NextFocusCurve.Zoom, PrevFocusPoint.Focus, NextTangent, 0.0);
 
-				InParams.GetDistortionState.ExecuteIfBound(PrevFocusPoint, PrevFocusCurve, Results.BlendingParams->States[0]);
-				InParams.GetDistortionState.ExecuteIfBound(PrevFocusPoint, NextFocusCurve, Results.BlendingParams->States[1]);
+				Results.BlendingParams->PatchCorners[0] = FDisplacementMapBlendPatchCorner(BlendCurve.X0, PrevFocusPoint.Focus, BlendCurve.Tangent0, 0.0);
+				Results.BlendingParams->PatchCorners[1] = FDisplacementMapBlendPatchCorner(BlendCurve.X1, PrevFocusPoint.Focus, BlendCurve.Tangent1, 0.0);
+
+				Results.BlendingParams->States[0] = InterpolationStates[0];
+				Results.BlendingParams->States[1] = InterpolationStates[1];
 			}
 
 			if (Results.BlendedDistortionParams.IsSet())
 			{
-				// Blend all distortion parameters using the map blend curve
-				TOptional<FDistortionInfo> PrevParameters = InParams.GetDistortionParameters.Execute(PrevFocusPoint, PrevFocusCurve);
-				TOptional<FDistortionInfo> NextParameters = InParams.GetDistortionParameters.Execute(PrevFocusPoint, NextFocusCurve);
-				
-				const int32 NumParams = Results.BlendedDistortionParams->Parameters.Num();
-				for (int Index = 0; Index < NumParams; ++Index)
-				{
-					BlendCurve.Y0 = PrevParameters.IsSet() ? PrevParameters.GetValue().Parameters[Index] : 0.0;
-					BlendCurve.Y1 = NextParameters.IsSet() ? NextParameters.GetValue().Parameters[Index] : 0.0;
-
-					Results.BlendedDistortionParams->Parameters[Index] = BlendCurve.Eval(InZoom);
-				}
+				Results.BlendedDistortionParams->Parameters = BlendCurve.MultiEval(InZoom, InterpolationStates[0].DistortionInfo.Parameters, InterpolationStates[0].DistortionInfo.Parameters);
 			}
 			
 			if (Results.UndistortedMaps.IsSet())
 			{
-				InParams.GetDisplacementMaps.ExecuteIfBound(PrevFocusPoint, PrevFocusCurve, (*Results.UndistortedMaps)[0], (*Results.DistortedMaps)[0]);
-				InParams.GetDisplacementMaps.ExecuteIfBound(PrevFocusPoint, NextFocusCurve, (*Results.UndistortedMaps)[1], (*Results.DistortedMaps)[1]);
+				GetDisplacementMaps<TableType>(PrevFocusPoint, BlendCurve.X0, (*Results.UndistortedMaps)[0], (*Results.DistortedMaps)[0]);
+				GetDisplacementMaps<TableType>(PrevFocusPoint, BlendCurve.X1, (*Results.UndistortedMaps)[1], (*Results.DistortedMaps)[1]);
 				
-				// Blend the overscan factor using the map blend curve
-				BlendCurve.Y0 = InParams.ProcessDisplacementMaps.Execute(PrevFocusPoint, PrevFocusCurve, (*Results.UndistortedMaps)[0], (*Results.DistortedMaps)[0]);
-				BlendCurve.Y1 = InParams.ProcessDisplacementMaps.Execute(PrevFocusPoint, NextFocusCurve, (*Results.UndistortedMaps)[1], (*Results.DistortedMaps)[1]);
+				if (InParams.ProcessDisplacementMaps.IsBound())
+				{
+					InParams.ProcessDisplacementMaps.Execute(InterpolationStates[0], (*Results.UndistortedMaps)[0], (*Results.DistortedMaps)[0]);
+					InParams.ProcessDisplacementMaps.Execute(InterpolationStates[1], (*Results.UndistortedMaps)[1], (*Results.DistortedMaps)[1]);
+				}
+			}
+
+			if (Results.BlendedOverscan.IsSet())
+			{
+				BlendCurve.Y0 = InParams.GetOverscan.Execute(PrevFocusPoint, BlendCurve.X0, InterpolationStates[0]);
+				BlendCurve.Y1 = InParams.GetOverscan.Execute(PrevFocusPoint, BlendCurve.X1, InterpolationStates[1]);
 				
 				Results.BlendedOverscan = BlendCurve.Eval(InZoom);
 			}
 		}
 		else if (CurveNeighbors.IsSinglePoint())
 		{
-			// We are on a focus curve edge of the blending patch
-			float PrevTangent, NextTangent;
-			SafeGetTangents(PrevFocusCurve.MapBlendingCurve, PointNeighbors, PrevTangent, NextTangent);
-			
-			FTangentBezierCurve BlendCurve(PrevFocusPoint.Focus, NextFocusPoint.Focus, 0.0f, 0.0f, PrevTangent, NextTangent);
+			FTangentBezierCurve BlendCurve = GetTangentCurve(PrevFocusCurve.MapBlendingCurve, InFocus, PrevFocusPoint.Focus, NextFocusPoint.Focus);
+
+			// Since the data is stored on the focus points and not the focus curves, we must find the focus points that have the actual defined point
+			// corresponding to the zoom of the focus curve we are on
+			LensDataTableUtils::FPointNeighbors OnCurvePointNeighbors = GetOnCurveFocusPoints<TableType>(FocusPoints, BlendCurve.X0, BlendCurve.X1);
+
+			const FocusPointType& OnCurvePrevFocusPoint = FocusPoints[OnCurvePointNeighbors.PreviousIndex];
+			const FocusPointType& OnCurveNextFocusPoint = FocusPoints[OnCurvePointNeighbors.NextIndex];
+
+			FCurveEndpoints CurveEndpoints = CreateCurveEndpoints(
+				GetDistortionInfo<TableType>(OnCurvePrevFocusPoint, PrevFocusCurve.Zoom),
+				GetDistortionInfo<TableType>(OnCurveNextFocusPoint, PrevFocusCurve.Zoom));
 			
+			FLensDistortionState InterpolationStates[2];
+			InterpolationStates[0].DistortionInfo = CurveEndpoints.X0;
+			InterpolationStates[1].DistortionInfo = CurveEndpoints.X1;
+
+			if (InParams.GetInterpolatedImageCenter.IsBound())
+			{
+				InParams.GetInterpolatedImageCenter.Execute(OnCurvePrevFocusPoint.Focus, PrevFocusCurve.Zoom, InterpolationStates[0].ImageCenter);
+				InParams.GetInterpolatedImageCenter.Execute(OnCurveNextFocusPoint.Focus, PrevFocusCurve.Zoom, InterpolationStates[1].ImageCenter);
+			}
+
+			if (InParams.GetInterpolatedFocalLength.IsBound())
+			{
+				InParams.GetInterpolatedFocalLength.Execute(OnCurvePrevFocusPoint.Focus, PrevFocusCurve.Zoom, InterpolationStates[0].FocalLengthInfo);
+				InParams.GetInterpolatedFocalLength.Execute(OnCurveNextFocusPoint.Focus, PrevFocusCurve.Zoom, InterpolationStates[1].FocalLengthInfo);
+			}
+
 			if (Results.BlendingParams.IsSet())
 			{
 				Results.BlendingParams->BlendType = EDisplacementMapBlendType::TwoFocusOneZoom;
-				Results.BlendingParams->PatchCorners[0] = FDisplacementMapBlendPatchCorner(PrevFocusCurve.Zoom, PrevFocusPoint.Focus, 0.0, PrevTangent);
-				Results.BlendingParams->PatchCorners[1] = FDisplacementMapBlendPatchCorner(PrevFocusCurve.Zoom, NextFocusPoint.Focus, 0.0, NextTangent);
 
-				InParams.GetDistortionState.ExecuteIfBound(PrevFocusPoint, PrevFocusCurve, Results.BlendingParams->States[0]);
-				InParams.GetDistortionState.ExecuteIfBound(NextFocusPoint, PrevFocusCurve, Results.BlendingParams->States[1]);
+				Results.BlendingParams->PatchCorners[0] = FDisplacementMapBlendPatchCorner(PrevFocusCurve.Zoom, OnCurvePrevFocusPoint.Focus, 0.0f, BlendCurve.Tangent0);
+				Results.BlendingParams->PatchCorners[1] = FDisplacementMapBlendPatchCorner(PrevFocusCurve.Zoom, OnCurveNextFocusPoint.Focus, 0.0f, BlendCurve.Tangent1);
+
+				Results.BlendingParams->States[0] = InterpolationStates[0];
+				Results.BlendingParams->States[1] = InterpolationStates[1];
 			}
 			
 			if (Results.BlendedDistortionParams.IsSet())
 			{
-				// Blend all distortion parameters using the map blend curve
-				TOptional<FDistortionInfo> PrevParameters = InParams.GetDistortionParameters.Execute(PrevFocusPoint, PrevFocusCurve);
-				TOptional<FDistortionInfo> NextParameters = InParams.GetDistortionParameters.Execute(NextFocusPoint, PrevFocusCurve);
-				
-				const int32 NumParams = Results.BlendedDistortionParams->Parameters.Num();
-				for (int Index = 0; Index < NumParams; ++Index)
-				{
-					BlendCurve.Y0 = PrevParameters.IsSet() ? PrevParameters.GetValue().Parameters[Index] : 0.0;
-					BlendCurve.Y1 = NextParameters.IsSet() ? NextParameters.GetValue().Parameters[Index] : 0.0;
-
-					Results.BlendedDistortionParams->Parameters[Index] = BlendCurve.Eval(InFocus);
-				}
+				Results.BlendedDistortionParams->Parameters = BlendCurve.MultiEval(InFocus, InterpolationStates[0].DistortionInfo.Parameters, InterpolationStates[1].DistortionInfo.Parameters);
 			}
 			
 			if (Results.UndistortedMaps.IsSet())
 			{
-				InParams.GetDisplacementMaps.ExecuteIfBound(PrevFocusPoint, PrevFocusCurve, (*Results.UndistortedMaps)[0], (*Results.DistortedMaps)[0]);
-				InParams.GetDisplacementMaps.ExecuteIfBound(NextFocusPoint, PrevFocusCurve, (*Results.UndistortedMaps)[1], (*Results.DistortedMaps)[1]);
-				
-				// Blend the overscan factor using the map blend curve
-				BlendCurve.Y0 = InParams.ProcessDisplacementMaps.Execute(PrevFocusPoint, PrevFocusCurve, (*Results.UndistortedMaps)[0], (*Results.DistortedMaps)[0]);
-				BlendCurve.Y1 = InParams.ProcessDisplacementMaps.Execute(NextFocusPoint, PrevFocusCurve, (*Results.UndistortedMaps)[1], (*Results.DistortedMaps)[1]);
-				
+				GetDisplacementMaps<TableType>(OnCurvePrevFocusPoint, PrevFocusCurve.Zoom, (*Results.UndistortedMaps)[0], (*Results.DistortedMaps)[0]);
+				GetDisplacementMaps<TableType>(OnCurveNextFocusPoint, PrevFocusCurve.Zoom, (*Results.UndistortedMaps)[1], (*Results.DistortedMaps)[1]);
+
+				if (InParams.ProcessDisplacementMaps.IsBound())
+				{
+					InParams.ProcessDisplacementMaps.Execute(InterpolationStates[0], (*Results.UndistortedMaps)[0], (*Results.DistortedMaps)[0]);
+					InParams.ProcessDisplacementMaps.Execute(InterpolationStates[1], (*Results.UndistortedMaps)[1], (*Results.DistortedMaps)[1]);
+				}
+			}
+
+			if (Results.BlendedOverscan.IsSet())
+			{
+				BlendCurve.Y0 = InParams.GetOverscan.Execute(OnCurvePrevFocusPoint, PrevFocusCurve.Zoom, InterpolationStates[0]);
+				BlendCurve.Y1 = InParams.GetOverscan.Execute(OnCurveNextFocusPoint, PrevFocusCurve.Zoom, InterpolationStates[1]);
+							
 				Results.BlendedOverscan = BlendCurve.Eval(InFocus);
 			}
 		}
 		else
 		{
 			// Otherwise, we are somewhere in the middle of the patch, and must do a full Coons patch blend
-			float XTangents[4];
-			SafeGetTangents(PrevFocusPoint.MapBlendingCurve, CurveNeighbors, XTangents[0], XTangents[1]);
-			SafeGetTangents(NextFocusPoint.MapBlendingCurve, CurveNeighbors, XTangents[2], XTangents[3]);
+
+			// We will definitely have four curves that intersect in a quad patch within which a Coons patch interpolation can be performed.
+			// However, the four corners of the patch may not all be defined and thus won't have tangent or distortion parameter data to interpolate
+			// from. These 'virtual' points must be approximated before the final Coons patch interpolation can be done. The curves can be used to
+			// interpolate the 'virtual' points from the nearest defined points, and from there, an approximate Coons patch can be assembled from which
+			// the final interpolation can be performed.
+			
+			FTangentBezierCoonsPatch CoonsPatch
+			(
+				FTangentBezierCoonsPatch::FPatchCurve(GetTangentCurve(PrevFocusPoint.MapBlendingCurve, InZoom, PrevFocusCurve.Zoom, NextFocusCurve.Zoom), PrevFocusPoint.Focus),
+				FTangentBezierCoonsPatch::FPatchCurve(GetTangentCurve(NextFocusPoint.MapBlendingCurve, InZoom, PrevFocusCurve.Zoom, NextFocusCurve.Zoom), NextFocusPoint.Focus),
+				FTangentBezierCoonsPatch::FPatchCurve(GetTangentCurve(PrevFocusCurve.MapBlendingCurve, InFocus, PrevFocusPoint.Focus, NextFocusPoint.Focus), PrevFocusCurve.Zoom),
+				FTangentBezierCoonsPatch::FPatchCurve(GetTangentCurve(NextFocusCurve.MapBlendingCurve, InFocus, PrevFocusPoint.Focus, NextFocusPoint.Focus), NextFocusCurve.Zoom)
+			);
+
+			// Interpolate the tangents needed at the potentially virtual patch corners from the curves from definite points
+			// Distortion parameter curves use the focus/zoom values on the curves as the curves' y values, so we must
+			// mirror this in order to properly interpolate the tangents at the virtual points
+			CoonsPatch.XCurves[0].TangentCurve.Y0 = CoonsPatch.XCurves[0].TangentCurve.X0;
+			CoonsPatch.XCurves[0].TangentCurve.Y1 = CoonsPatch.XCurves[0].TangentCurve.X1;
+			CoonsPatch.XCurves[1].TangentCurve.Y0 = CoonsPatch.XCurves[1].TangentCurve.X0;
+			CoonsPatch.XCurves[1].TangentCurve.Y1 = CoonsPatch.XCurves[1].TangentCurve.X1;
+			CoonsPatch.YCurves[0].TangentCurve.Y0 = CoonsPatch.YCurves[0].TangentCurve.X0;
+			CoonsPatch.YCurves[0].TangentCurve.Y1 = CoonsPatch.YCurves[0].TangentCurve.X1;
+			CoonsPatch.YCurves[1].TangentCurve.Y0 = CoonsPatch.YCurves[1].TangentCurve.X0;
+			CoonsPatch.YCurves[1].TangentCurve.Y1 = CoonsPatch.YCurves[1].TangentCurve.X1;
+			
+			float XTangents[4] =
+			{
+				CoonsPatch.XCurves[0].TangentCurve.EvalTangent(PrevFocusCurve.Zoom),
+				CoonsPatch.XCurves[0].TangentCurve.EvalTangent(NextFocusCurve.Zoom),
+				CoonsPatch.XCurves[1].TangentCurve.EvalTangent(NextFocusCurve.Zoom),
+				CoonsPatch.XCurves[1].TangentCurve.EvalTangent(PrevFocusCurve.Zoom)
+			};
 		
-			float YTangents[4];
-			SafeGetTangents(PrevFocusCurve.MapBlendingCurve, PointNeighbors, YTangents[0], YTangents[3]);
-			SafeGetTangents(NextFocusCurve.MapBlendingCurve, PointNeighbors, YTangents[1], YTangents[2]);
+			float YTangents[4] =
+			{
+				CoonsPatch.YCurves[0].TangentCurve.EvalTangent(PrevFocusPoint.Focus),
+				CoonsPatch.YCurves[1].TangentCurve.EvalTangent(PrevFocusPoint.Focus),
+				CoonsPatch.YCurves[1].TangentCurve.EvalTangent(NextFocusPoint.Focus),
+				CoonsPatch.YCurves[0].TangentCurve.EvalTangent(NextFocusPoint.Focus)
+			};
+
+			LensDataTableUtils::FPointNeighbors PrevCurvePointNeighbors = GetOnCurveFocusPoints<TableType>(FocusPoints, CoonsPatch.YCurves[0].TangentCurve.X0, CoonsPatch.YCurves[0].TangentCurve.X1);
+			LensDataTableUtils::FPointNeighbors NextCurvePointNeighbors = GetOnCurveFocusPoints<TableType>(FocusPoints, CoonsPatch.YCurves[1].TangentCurve.X0, CoonsPatch.YCurves[1].TangentCurve.X1);
+			
+			const FocusPointType& PrevCurvePrevFocusPoint = FocusPoints[PrevCurvePointNeighbors.PreviousIndex];
+			const FocusPointType& PrevCurveNextFocusPoint = FocusPoints[PrevCurvePointNeighbors.NextIndex];
+			const FocusPointType& NextCurvePrevFocusPoint = FocusPoints[NextCurvePointNeighbors.PreviousIndex];
+			const FocusPointType& NextCurveNextFocusPoint = FocusPoints[NextCurvePointNeighbors.NextIndex];
+
+			FLensDistortionState InterpolationStates[4];
+			FDistortionInfo InterpolatedDistortionInfo;
+
+			// Compute the virtual corners of the patch via interpolation from the nearest defined points on each edge curve
+			FCurveEndpoints XCurveEndpoints[2] =
+            {
+                CreateCurveEndpoints(
+                	GetDistortionInfo<TableType>(PrevFocusPoint, CoonsPatch.XCurves[0].TangentCurve.X0),
+                	GetDistortionInfo<TableType>(PrevFocusPoint, CoonsPatch.XCurves[0].TangentCurve.X1)
+                ),
+                CreateCurveEndpoints(
+                	GetDistortionInfo<TableType>(NextFocusPoint, CoonsPatch.XCurves[1].TangentCurve.X0),
+					GetDistortionInfo<TableType>(NextFocusPoint, CoonsPatch.XCurves[1].TangentCurve.X1)
+                )
+            };
+
+			FCurveEndpoints YCurveEndpoints[2] =
+			{
+				CreateCurveEndpoints(
+					GetDistortionInfo<TableType>(PrevCurvePrevFocusPoint, PrevFocusCurve.Zoom),
+					GetDistortionInfo<TableType>(PrevCurveNextFocusPoint, PrevFocusCurve.Zoom)
+				),
+				CreateCurveEndpoints(
+					GetDistortionInfo<TableType>(NextCurvePrevFocusPoint, NextFocusCurve.Zoom),
+					GetDistortionInfo<TableType>(NextCurveNextFocusPoint, NextFocusCurve.Zoom)
+				)
+			};
+
+			const int32 NumParams = InParams.DefaultDistortionParams.Num();
+			InterpolationStates[0].DistortionInfo.Parameters.AddZeroed(NumParams);
+			InterpolationStates[1].DistortionInfo.Parameters.AddZeroed(NumParams);
+			InterpolationStates[2].DistortionInfo.Parameters.AddZeroed(NumParams);
+			InterpolationStates[3].DistortionInfo.Parameters.AddZeroed(NumParams);
+			InterpolatedDistortionInfo.Parameters.AddZeroed(NumParams);
+			
+			for (int Index = 0; Index < NumParams; ++Index)
+			{
+				CoonsPatch.XCurves[0].TangentCurve.Y0 = XCurveEndpoints[0].X0.Parameters[Index];
+				CoonsPatch.XCurves[0].TangentCurve.Y1 = XCurveEndpoints[0].X1.Parameters[Index];
+				CoonsPatch.XCurves[1].TangentCurve.Y0 = XCurveEndpoints[1].X0.Parameters[Index];
+				CoonsPatch.XCurves[1].TangentCurve.Y1 = XCurveEndpoints[1].X1.Parameters[Index];
+
+				CoonsPatch.YCurves[0].TangentCurve.Y0 = YCurveEndpoints[0].X0.Parameters[Index];
+				CoonsPatch.YCurves[0].TangentCurve.Y1 = YCurveEndpoints[0].X1.Parameters[Index];
+				CoonsPatch.YCurves[1].TangentCurve.Y0 = YCurveEndpoints[1].X0.Parameters[Index];
+				CoonsPatch.YCurves[1].TangentCurve.Y1 = YCurveEndpoints[1].X1.Parameters[Index];
+
+				InterpolationStates[0].DistortionInfo.Parameters[Index] = CoonsPatch.Blend(FVector2D(PrevFocusCurve.Zoom, PrevFocusPoint.Focus));
+				InterpolationStates[1].DistortionInfo.Parameters[Index] = CoonsPatch.Blend(FVector2D(NextFocusCurve.Zoom, PrevFocusPoint.Focus));
+				InterpolationStates[2].DistortionInfo.Parameters[Index] = CoonsPatch.Blend(FVector2D(NextFocusCurve.Zoom, NextFocusPoint.Focus));
+				InterpolationStates[3].DistortionInfo.Parameters[Index] = CoonsPatch.Blend(FVector2D(PrevFocusCurve.Zoom, NextFocusPoint.Focus));
+				
+				InterpolatedDistortionInfo.Parameters[Index] = CoonsPatch.Blend(FVector2D(InZoom, InFocus));
+			}
+			
+			if (InParams.GetInterpolatedImageCenter.IsBound())
+			{
+				InParams.GetInterpolatedImageCenter.Execute(PrevFocusPoint.Focus, PrevFocusCurve.Zoom, InterpolationStates[0].ImageCenter);
+				InParams.GetInterpolatedImageCenter.Execute(PrevFocusPoint.Focus, NextFocusCurve.Zoom, InterpolationStates[1].ImageCenter);
+				InParams.GetInterpolatedImageCenter.Execute(NextFocusPoint.Focus, NextFocusCurve.Zoom, InterpolationStates[2].ImageCenter);
+				InParams.GetInterpolatedImageCenter.Execute(NextFocusPoint.Focus, PrevFocusCurve.Zoom, InterpolationStates[3].ImageCenter);
+			}
+			
+			if (InParams.GetInterpolatedFocalLength.IsBound())
+			{
+				InParams.GetInterpolatedFocalLength.Execute(PrevFocusPoint.Focus, PrevFocusCurve.Zoom, InterpolationStates[0].FocalLengthInfo);
+				InParams.GetInterpolatedFocalLength.Execute(PrevFocusPoint.Focus, NextFocusCurve.Zoom, InterpolationStates[1].FocalLengthInfo);
+				InParams.GetInterpolatedFocalLength.Execute(NextFocusPoint.Focus, NextFocusCurve.Zoom, InterpolationStates[2].FocalLengthInfo);
+				InParams.GetInterpolatedFocalLength.Execute(NextFocusPoint.Focus, PrevFocusCurve.Zoom, InterpolationStates[3].FocalLengthInfo);
+			}
 			
-			FTangentBezierCoonsPatch CoonsPatch(PrevFocusCurve.Zoom, NextFocusCurve.Zoom, PrevFocusPoint.Focus, NextFocusPoint.Focus, XTangents, YTangents);
-
 			if (Results.BlendingParams.IsSet())
 			{
 				Results.BlendingParams->BlendType = EDisplacementMapBlendType::TwoFocusTwoZoom;
+
 				Results.BlendingParams->PatchCorners[0] = FDisplacementMapBlendPatchCorner(PrevFocusCurve.Zoom, PrevFocusPoint.Focus, XTangents[0], YTangents[0]);
 				Results.BlendingParams->PatchCorners[1] = FDisplacementMapBlendPatchCorner(NextFocusCurve.Zoom, PrevFocusPoint.Focus, XTangents[1], YTangents[1]);
 				Results.BlendingParams->PatchCorners[2] = FDisplacementMapBlendPatchCorner(NextFocusCurve.Zoom, NextFocusPoint.Focus, XTangents[2], YTangents[2]);
 				Results.BlendingParams->PatchCorners[3] = FDisplacementMapBlendPatchCorner(PrevFocusCurve.Zoom, NextFocusPoint.Focus, XTangents[3], YTangents[3]);
 
-				InParams.GetDistortionState.ExecuteIfBound(PrevFocusPoint, PrevFocusCurve, Results.BlendingParams->States[0]);
-				InParams.GetDistortionState.ExecuteIfBound(PrevFocusPoint, NextFocusCurve, Results.BlendingParams->States[1]);
-				InParams.GetDistortionState.ExecuteIfBound(NextFocusPoint, NextFocusCurve, Results.BlendingParams->States[2]);
-				InParams.GetDistortionState.ExecuteIfBound(NextFocusPoint, PrevFocusCurve, Results.BlendingParams->States[3]);
+				Results.BlendingParams->States[0] = InterpolationStates[0];
+				Results.BlendingParams->States[1] = InterpolationStates[1];
+				Results.BlendingParams->States[2] = InterpolationStates[2];
+				Results.BlendingParams->States[3] = InterpolationStates[3];
 			}
 			
 			if (Results.BlendedDistortionParams.IsSet())
 			{
-				TOptional<FDistortionInfo> Points[4] =
-				{
-					InParams.GetDistortionParameters.Execute(PrevFocusPoint, PrevFocusCurve),
-					InParams.GetDistortionParameters.Execute(PrevFocusPoint, NextFocusCurve),
-					InParams.GetDistortionParameters.Execute(NextFocusPoint, NextFocusCurve),
-					InParams.GetDistortionParameters.Execute(NextFocusPoint, PrevFocusCurve)
-				};
+				Results.BlendedDistortionParams = InterpolatedDistortionInfo;
+			}
+			
+			if (Results.UndistortedMaps.IsSet())
+			{
+				GetDisplacementMaps<TableType>(PrevFocusPoint, PrevFocusCurve.Zoom, (*Results.UndistortedMaps)[0], (*Results.DistortedMaps)[0]);
+				GetDisplacementMaps<TableType>(PrevFocusPoint, NextFocusCurve.Zoom, (*Results.UndistortedMaps)[1], (*Results.DistortedMaps)[1]);
+				GetDisplacementMaps<TableType>(NextFocusPoint, NextFocusCurve.Zoom, (*Results.UndistortedMaps)[2], (*Results.DistortedMaps)[2]);
+				GetDisplacementMaps<TableType>(NextFocusPoint, PrevFocusCurve.Zoom, (*Results.UndistortedMaps)[3], (*Results.DistortedMaps)[3]);
 
-				// Handle edge case for when a corner point of the patch is not defined. If such a case happens, we
-				// average its value from its two neighbor points, which are guaranteed to exist
-				for (int32 Index = 0; Index < 4; ++Index)
+				if (InParams.ProcessDisplacementMaps.IsBound())
 				{
-					if (!Points[Index].IsSet())
-					{
-						int32 PrevIndex = Index > 0 ? Index - 1 : 3;
-						int32 NextIndex = Index < 3 ? Index + 1 : 0;
-
-						Points[Index] = FDistortionInfo();
-						
-						int32 NumParams = Results.BlendedDistortionParams->Parameters.Num();
-						for (int ParamIndex = 0; ParamIndex < NumParams; ++ParamIndex)
-						{
-							Points[Index]->Parameters.Add(0.5 * (Points[PrevIndex]->Parameters[ParamIndex] + Points[NextIndex]->Parameters[ParamIndex]));
-						}
-					}
+					InParams.ProcessDisplacementMaps.Execute(InterpolationStates[0], (*Results.UndistortedMaps)[0], (*Results.DistortedMaps)[0]);
+					InParams.ProcessDisplacementMaps.Execute(InterpolationStates[1], (*Results.UndistortedMaps)[1], (*Results.DistortedMaps)[1]);
+					InParams.ProcessDisplacementMaps.Execute(InterpolationStates[2], (*Results.UndistortedMaps)[2], (*Results.DistortedMaps)[2]);
+					InParams.ProcessDisplacementMaps.Execute(InterpolationStates[3], (*Results.UndistortedMaps)[3], (*Results.DistortedMaps)[3]);
 				}
+			}
+
+			if (Results.BlendedOverscan.IsSet())
+			{
+				// Since overscan is computed from distortion parameters, we only want to compute it for the distortion states for each of the potentially virtual patch corners.
+				// As such, it needs its own Coons patch where each curve starts and ends at the correct corners. The interpolated tangents computed earlier can be used here
+				// for each of the corners.
 				
-				const int32 NumParams = Results.BlendedDistortionParams->Parameters.Num();
-				for (int Index = 0; Index < NumParams; ++Index)
+				float Overscans[4] =
 				{
-					CoonsPatch.Corners[0].Value = Points[0]->Parameters[Index];
-					CoonsPatch.Corners[1].Value = Points[1]->Parameters[Index];
-					CoonsPatch.Corners[2].Value = Points[2]->Parameters[Index];
-					CoonsPatch.Corners[3].Value = Points[3]->Parameters[Index];
+					InParams.GetOverscan.Execute(PrevFocusPoint, PrevFocusCurve.Zoom, InterpolationStates[0]),
+					InParams.GetOverscan.Execute(PrevFocusPoint, NextFocusCurve.Zoom, InterpolationStates[1]),
+					InParams.GetOverscan.Execute(NextFocusPoint, NextFocusCurve.Zoom, InterpolationStates[2]),
+					InParams.GetOverscan.Execute(NextFocusPoint, PrevFocusCurve.Zoom, InterpolationStates[3]),
+				};
 				
-					Results.BlendedDistortionParams->Parameters[Index] = CoonsPatch.Blend(FVector2D(InZoom, InFocus));
-				}
-			}
-			
-			if (Results.UndistortedMaps.IsSet())
-			{
-				InParams.GetDisplacementMaps.ExecuteIfBound(PrevFocusPoint, PrevFocusCurve, (*Results.UndistortedMaps)[0], (*Results.DistortedMaps)[0]);
-				InParams.GetDisplacementMaps.ExecuteIfBound(PrevFocusPoint, NextFocusCurve, (*Results.UndistortedMaps)[1], (*Results.DistortedMaps)[1]);
-				InParams.GetDisplacementMaps.ExecuteIfBound(NextFocusPoint, NextFocusCurve, (*Results.UndistortedMaps)[2], (*Results.DistortedMaps)[2]);
-				InParams.GetDisplacementMaps.ExecuteIfBound(NextFocusPoint, PrevFocusCurve, (*Results.UndistortedMaps)[3], (*Results.DistortedMaps)[3]);
+				FTangentBezierCoonsPatch OverscanCoonsPatch
+				(
+					FTangentBezierCoonsPatch::FPatchCurve(FTangentBezierCurve(PrevFocusCurve.Zoom, NextFocusCurve.Zoom, Overscans[0], Overscans[1], XTangents[0], XTangents[1]), PrevFocusPoint.Focus),
+					FTangentBezierCoonsPatch::FPatchCurve(FTangentBezierCurve(PrevFocusCurve.Zoom, NextFocusCurve.Zoom, Overscans[3], Overscans[2], XTangents[3], XTangents[2]), NextFocusPoint.Focus),
+					FTangentBezierCoonsPatch::FPatchCurve(FTangentBezierCurve(PrevFocusPoint.Focus, NextFocusPoint.Focus, Overscans[0], Overscans[3], YTangents[0], YTangents[3]), PrevFocusCurve.Zoom),
+					FTangentBezierCoonsPatch::FPatchCurve(FTangentBezierCurve(PrevFocusPoint.Focus, NextFocusPoint.Focus, Overscans[1], Overscans[2], YTangents[1], YTangents[2]), NextFocusCurve.Zoom)
+				);
 				
-				CoonsPatch.Corners[0].Value = InParams.ProcessDisplacementMaps.Execute(PrevFocusPoint, PrevFocusCurve, (*Results.UndistortedMaps)[0], (*Results.DistortedMaps)[0]);
-				CoonsPatch.Corners[1].Value = InParams.ProcessDisplacementMaps.Execute(PrevFocusPoint, NextFocusCurve, (*Results.UndistortedMaps)[1], (*Results.DistortedMaps)[1]);
-				CoonsPatch.Corners[2].Value = InParams.ProcessDisplacementMaps.Execute(NextFocusPoint, NextFocusCurve, (*Results.UndistortedMaps)[2], (*Results.DistortedMaps)[2]);
-				CoonsPatch.Corners[3].Value = InParams.ProcessDisplacementMaps.Execute(NextFocusPoint, PrevFocusCurve, (*Results.UndistortedMaps)[3], (*Results.DistortedMaps)[3]);
-
-				Results.BlendedOverscan = CoonsPatch.Blend(FVector2D(InZoom, InFocus));
+				Results.BlendedOverscan = OverscanCoonsPatch.Blend(FVector2D(InZoom, InFocus));
 			}
 		}
 
