CameraCalibrationCore: Use principal point instead of image center for distortion normalization.

This convention differs from OpenTrackIO which can have explicit projection vs distortion offsets.

#rnx

#virtualized
#Synced-CL 51046329

diff --git a/Engine/Plugins/VirtualProduction/CameraCalibrationCore/Shaders/Private/SphericalDistortion.usf b/Engine/Plugins/VirtualProduction/CameraCalibrationCore/Shaders/Private/SphericalDistortion.usf
--- a/Engine/Plugins/VirtualProduction/CameraCalibrationCore/Shaders/Private/SphericalDistortion.usf
+++ b/Engine/Plugins/VirtualProduction/CameraCalibrationCore/Shaders/Private/SphericalDistortion.usf
@@ -24,10 +24,8 @@ float P2;
 
 float2 SphericalDistort(float2 InputUV)
 {
-	// @todo Distortion offset to replace the 0.5 below
-
-	float x_prime = (InputUV.x - 0.5) / FocalLength.x;
-	float y_prime = (InputUV.y - 0.5) / FocalLength.y;
+	float x_prime = (InputUV.x - ImageCenter.x) / FocalLength.x;
+	float y_prime = (InputUV.y - ImageCenter.y) / FocalLength.y;
 
 	float x_prime2 = x_prime * x_prime;
 	float y_prime2 = y_prime * y_prime;
diff --git a/Engine/Plugins/VirtualProduction/CameraCalibrationCore/Source/CameraCalibrationCore/Private/SphericalLensDistortionModelHandler.cpp b/Engine/Plugins/VirtualProduction/CameraCalibrationCore/Source/CameraCalibrationCore/Private/SphericalLensDistortionModelHandler.cpp
--- a/Engine/Plugins/VirtualProduction/CameraCalibrationCore/Source/CameraCalibrationCore/Private/SphericalLensDistortionModelHandler.cpp
+++ b/Engine/Plugins/VirtualProduction/CameraCalibrationCore/Source/CameraCalibrationCore/Private/SphericalLensDistortionModelHandler.cpp
@@ -62,8 +62,7 @@ FVector2D USphericalLensDistortionModelHandler::ComputeDistortedUV(const FVector
 		return InScreenUV;
 	}
 
-	// @todo distortion offset to replace the 0.5 below
-	FVector2D X = (InScreenUV - FVector2D(0.5f, 0.5f)) / FxFy; // normalized coords (undistorted)
+	FVector2D X = (InScreenUV - CxCy) / FxFy;
 
 	const float R2 = X.X * X.X + X.Y * X.Y;
 	const float R4 = R2 * R2;
@@ -121,8 +120,7 @@ FVector2D USphericalLensDistortionModelHandler::ComputeUndistortedUV(const FVect
 		P = (Pd - Tangential) * InvRadial;
 	}
 
-	// @todo distortion offset to replace the 0.5 below
-	return P * FxFy + FVector2D(0.5f, 0.5f); // undistorted UV
+	return P * FxFy + CxCy;
 }
 
 void USphericalLensDistortionModelHandler::InitDistortionMaterials()
