CameraCalibrationCore: Back to baking cx cy in spherical distortion shaders.

This should avoid the need to manipulate filmback offset (although doing that has the advantage of not needing overscan to compensate for the shift).

#rnx

#virtualized
#Synced-CL 51046329

diff --git a/Engine/Plugins/VirtualProduction/CameraCalibrationCore/Shaders/Private/SphericalDistortion.usf b/Engine/Plugins/VirtualProduction/CameraCalibrationCore/Shaders/Private/SphericalDistortion.usf
--- a/Engine/Plugins/VirtualProduction/CameraCalibrationCore/Shaders/Private/SphericalDistortion.usf
+++ b/Engine/Plugins/VirtualProduction/CameraCalibrationCore/Shaders/Private/SphericalDistortion.usf
@@ -94,7 +94,10 @@ void MainCS(uint3 DispatchThreadId : SV_DispatchThreadID)
 	UndistortedUV = ((UndistortedUV - 0.5) * InverseOverscan * CameraOverscan) + 0.5;
 	
 	float2 DistortedUV = TryDistort(UndistortedUV);
-	
-	// Compute the UV displacement and scale it by the CameraOverscan to account for the increased resolution of the overscanned scene texture 
+
+	// Bake principal point projection offset (sensor shift)
+	DistortedUV += ImageCenter - 0.5;
+
+	// Compute the UV displacement and scale it by the CameraOverscan to account for the increased resolution of the overscanned scene texture
 	OutDistortionMap[DispatchThreadId.xy] = (DistortedUV - UndistortedUV) / CameraOverscan;
 }
\ No newline at end of file
diff --git a/Engine/Plugins/VirtualProduction/CameraCalibrationCore/Source/CameraCalibrationCore/Private/SphericalLensDistortionModelHandler.cpp b/Engine/Plugins/VirtualProduction/CameraCalibrationCore/Source/CameraCalibrationCore/Private/SphericalLensDistortionModelHandler.cpp
--- a/Engine/Plugins/VirtualProduction/CameraCalibrationCore/Source/CameraCalibrationCore/Private/SphericalLensDistortionModelHandler.cpp
+++ b/Engine/Plugins/VirtualProduction/CameraCalibrationCore/Source/CameraCalibrationCore/Private/SphericalLensDistortionModelHandler.cpp
@@ -80,7 +80,8 @@ FVector2D USphericalLensDistortionModelHandler::ComputeDistortedUV(const FVector
 
 	X = X * Radial + Tangential;
 
-	return X * FxFy + CxCy; // distorted UV
+	// Distortion reprojection + sensor shift
+	return X * FxFy + CxCy + (CxCy - FVector2D(0.5f, 0.5f));
 }
 
 FVector2D USphericalLensDistortionModelHandler::ComputeUndistortedUV(const FVector2D& DistortedUV) const
@@ -93,7 +94,9 @@ FVector2D USphericalLensDistortionModelHandler::ComputeUndistortedUV(const FVect
 		return DistortedUV;
 	}
 
-	const FVector2D Pd = (DistortedUV - CxCy) / FxFy;
+	// Undo sensor shift before undistorting
+	const FVector2D Unshifted = DistortedUV - (CxCy - FVector2D(0.5f, 0.5f));
+	const FVector2D Pd = (Unshifted - CxCy) / FxFy;
 
 	FVector2D P = Pd; // initial guess
 
