Camera Calibration: Revert shift reorder.

#rnx

#virtualized
#Synced-CL 51080076

diff --git a/Engine/Plugins/VirtualProduction/CameraCalibrationCore/Shaders/Private/SphericalDistortion.usf b/Engine/Plugins/VirtualProduction/CameraCalibrationCore/Shaders/Private/SphericalDistortion.usf
--- a/Engine/Plugins/VirtualProduction/CameraCalibrationCore/Shaders/Private/SphericalDistortion.usf
+++ b/Engine/Plugins/VirtualProduction/CameraCalibrationCore/Shaders/Private/SphericalDistortion.usf
@@ -24,8 +24,10 @@ float P2;
 
 float2 SphericalDistort(float2 InputUV)
 {
-	float x_prime = (InputUV.x - ImageCenter.x) / FocalLength.x;
-	float y_prime = (InputUV.y - ImageCenter.y) / FocalLength.y;
+	// @todo Distortion offset to replace the 0.5 below
+
+	float x_prime = (InputUV.x - 0.5) / FocalLength.x;
+	float y_prime = (InputUV.y - 0.5) / FocalLength.y;
 
 	float x_prime2 = x_prime * x_prime;
 	float y_prime2 = y_prime * y_prime;
@@ -94,10 +96,7 @@ void MainCS(uint3 DispatchThreadId : SV_DispatchThreadID)
 	UndistortedUV = ((UndistortedUV - 0.5) * InverseOverscan * CameraOverscan) + 0.5;
 	
 	float2 DistortedUV = TryDistort(UndistortedUV);
-
-	// Bake principal point projection offset (sensor shift)
-	DistortedUV += ImageCenter - 0.5;
-
-	// Compute the UV displacement and scale it by the CameraOverscan to account for the increased resolution of the overscanned scene texture
+	
+	// Compute the UV displacement and scale it by the CameraOverscan to account for the increased resolution of the overscanned scene texture 
 	OutDistortionMap[DispatchThreadId.xy] = (DistortedUV - UndistortedUV) / CameraOverscan;
 }
\ No newline at end of file
diff --git a/Engine/Plugins/VirtualProduction/CameraCalibrationCore/Source/CameraCalibrationCore/Private/SphericalLensDistortionModelHandler.cpp b/Engine/Plugins/VirtualProduction/CameraCalibrationCore/Source/CameraCalibrationCore/Private/SphericalLensDistortionModelHandler.cpp
--- a/Engine/Plugins/VirtualProduction/CameraCalibrationCore/Source/CameraCalibrationCore/Private/SphericalLensDistortionModelHandler.cpp
+++ b/Engine/Plugins/VirtualProduction/CameraCalibrationCore/Source/CameraCalibrationCore/Private/SphericalLensDistortionModelHandler.cpp
@@ -62,7 +62,8 @@ FVector2D USphericalLensDistortionModelHandler::ComputeDistortedUV(const FVector
 		return InScreenUV;
 	}
 
-	FVector2D X = (InScreenUV - CxCy) / FxFy;
+	// @todo distortion offset to replace the 0.5 below
+	FVector2D X = (InScreenUV - FVector2D(0.5f, 0.5f)) / FxFy; // normalized coords (undistorted)
 
 	const float R2 = X.X * X.X + X.Y * X.Y;
 	const float R4 = R2 * R2;
@@ -80,8 +81,7 @@ FVector2D USphericalLensDistortionModelHandler::ComputeDistortedUV(const FVector
 
 	X = X * Radial + Tangential;
 
-	// Distortion reprojection + sensor shift
-	return X * FxFy + CxCy + (CxCy - FVector2D(0.5f, 0.5f));
+	return X * FxFy + CxCy; // distorted UV
 }
 
 FVector2D USphericalLensDistortionModelHandler::ComputeUndistortedUV(const FVector2D& DistortedUV) const
@@ -94,9 +94,7 @@ FVector2D USphericalLensDistortionModelHandler::ComputeUndistortedUV(const FVect
 		return DistortedUV;
 	}
 
-	// Undo sensor shift before undistorting
-	const FVector2D Unshifted = DistortedUV - (CxCy - FVector2D(0.5f, 0.5f));
-	const FVector2D Pd = (Unshifted - CxCy) / FxFy;
+	const FVector2D Pd = (DistortedUV - CxCy) / FxFy;
 
 	FVector2D P = Pd; // initial guess
 
@@ -123,7 +121,8 @@ FVector2D USphericalLensDistortionModelHandler::ComputeUndistortedUV(const FVect
 		P = (Pd - Tangential) * InvRadial;
 	}
 
-	return P * FxFy + CxCy;
+	// @todo distortion offset to replace the 0.5 below
+	return P * FxFy + FVector2D(0.5f, 0.5f); // undistorted UV
 }
 
 void USphericalLensDistortionModelHandler::InitDistortionMaterials()
